Changing FPakListEntry < operator to be FORCEINLINE

diff --git a/Engine/Source/Runtime/PakFile/Public/IPlatformFilePak.h b/Engine/Source/Runtime/PakFile/Public/IPlatformFilePak.h
--- a/Engine/Source/Runtime/PakFile/Public/IPlatformFilePak.h
+++ b/Engine/Source/Runtime/PakFile/Public/IPlatformFilePak.h
@@ -784,7 +784,7 @@ class PAKFILE_API FPakPlatformFile : public IPlatformFile
 		uint32		ReadOrder;
 		FPakFile*	PakFile;
 
-		bool operator < (const FPakListEntry& RHS) const
+		FORCEINLINE bool operator < (const FPakListEntry& RHS) const
 		{
 			return ReadOrder > RHS.ReadOrder;
 		}
