Allow initialization of a PakFile from an existing Archive

diff --git a/Engine/Source/Runtime/PakFile/Private/IPlatformFilePak.cpp b/Engine/Source/Runtime/PakFile/Private/IPlatformFilePak.cpp
--- a/Engine/Source/Runtime/PakFile/Private/IPlatformFilePak.cpp
+++ b/Engine/Source/Runtime/PakFile/Private/IPlatformFilePak.cpp
@@ -256,6 +256,13 @@ FPakFile::FPakFile(IPlatformFile* LowerLevel, const TCHAR* Filename, bool bIsSig
 	}
 }
 
+FPakFile::FPakFile(FArchive* Archive)
+	: bSigned(false)
+	, bIsValid(false)
+{
+	Initialize(Archive);
+}
+
 FPakFile::~FPakFile()
 {
 }
diff --git a/Engine/Source/Runtime/PakFile/Public/IPlatformFilePak.h b/Engine/Source/Runtime/PakFile/Public/IPlatformFilePak.h
--- a/Engine/Source/Runtime/PakFile/Public/IPlatformFilePak.h
+++ b/Engine/Source/Runtime/PakFile/Public/IPlatformFilePak.h
@@ -296,6 +296,13 @@ public:
 	 */
 	FPakFile(IPlatformFile* LowerLevel, const TCHAR* Filename, bool bIsSigned);
 
+	/**
+	 * Creates a pak file using the supplied archive
+	 *
+	 * @param Archive	Pointer to the archive which contains the pak file data
+	 */
+	FPakFile(FArchive* Archive);
+
 	~FPakFile();
 
 	/**
