Fix modular server builds.

- The project file (passed with abs path) could be found in the pak file, but its directory could not.
- Fixed by PeterK in UT4 branch.

(Merging CL 2345332 from UT4 branch to main).

diff --git a/Engine/Source/Runtime/PakFile/Public/IPlatformFilePak.h b/Engine/Source/Runtime/PakFile/Public/IPlatformFilePak.h
--- a/Engine/Source/Runtime/PakFile/Public/IPlatformFilePak.h
+++ b/Engine/Source/Runtime/PakFile/Public/IPlatformFilePak.h
@@ -1094,14 +1094,18 @@ public:
 		{
 			const FString Path(FPaths::GetPath(Filename));
 			const FPakDirectory* PakDirectory = PakFile->FindDirectory(*Path);
-			const FString* RealFilename = PakDirectory->FindKey(const_cast<FPakEntry*>(FileEntry));
-			return *RealFilename;
-		}
-		else
-		{
-			// Fall back to lower level.
-			return LowerLevel->GetFilenameOnDisk(Filename);
+			if (PakDirectory != nullptr)
+			{
+				const FString* RealFilename = PakDirectory->FindKey(const_cast<FPakEntry*>(FileEntry));
+				if (RealFilename != nullptr)
+				{
+					return *RealFilename;
+				}
+			}
 		}
+
+		// Fall back to lower level.
+		return LowerLevel->GetFilenameOnDisk(Filename);
 	}
 
 	virtual IFileHandle* OpenRead(const TCHAR* Filename) override;
