Merging CL 2440669, unmount pak support

//depot/UE4-UT/...

to //depot/UE4/...

diff --git a/Engine/Source/Runtime/Core/Private/Misc/CallbackDevice.cpp b/Engine/Source/Runtime/Core/Private/Misc/CallbackDevice.cpp
--- a/Engine/Source/Runtime/Core/Private/Misc/CallbackDevice.cpp
+++ b/Engine/Source/Runtime/Core/Private/Misc/CallbackDevice.cpp
@@ -33,6 +33,7 @@ FCoreDelegates::FGetSerialNumberBlocksForDebugVisualizersDelegate& FCoreDelegate
 }
 
 FCoreDelegates::FOnMountPak FCoreDelegates::OnMountPak;
+FCoreDelegates::FOnUnmountPak FCoreDelegates::OnUnmountPak;
 FCoreDelegates::FOnUserLoginChangedEvent FCoreDelegates::OnUserLoginChangedEvent; 
 FCoreDelegates::FOnUserControllerConnectionChange FCoreDelegates::OnControllerConnectionChange;
 FCoreDelegates::FOnSafeFrameChangedEvent FCoreDelegates::OnSafeFrameChangedEvent;
diff --git a/Engine/Source/Runtime/Core/Public/Misc/CallbackDevice.h b/Engine/Source/Runtime/Core/Public/Misc/CallbackDevice.h
--- a/Engine/Source/Runtime/Core/Public/Misc/CallbackDevice.h
+++ b/Engine/Source/Runtime/Core/Public/Misc/CallbackDevice.h
@@ -51,6 +51,9 @@ public:
 	// delegate type for prompting the pak system to mount a new pak
 	DECLARE_DELEGATE_RetVal_TwoParams(bool, FOnMountPak, const FString&, uint32);
 
+	// delegate type for prompting the pak system to mount a new pak
+	DECLARE_DELEGATE_RetVal_OneParam(bool, FOnUnmountPak, const FString&);
+
 	// Callback for PER_MODULE_BOILERPLATE macro's GSerialNumberBlocksForDebugVisualizers
 	DECLARE_DELEGATE_RetVal(int32***, FGetSerialNumberBlocksForDebugVisualizersDelegate);
 
@@ -98,6 +101,9 @@ public:
 	// Callback for mounting a new pak file.
 	static FOnMountPak OnMountPak;
 
+	// Callback for unmounting a pak file.
+	static FOnUnmountPak OnUnmountPak;
+
 	// Callback when an ensure has occurred
 	static FOnHandleSystemEnsure OnHandleSystemEnsure;
 	// Callback when an error (crash) has occurred
diff --git a/Engine/Source/Runtime/PakFile/Private/IPlatformFilePak.cpp b/Engine/Source/Runtime/PakFile/Private/IPlatformFilePak.cpp
--- a/Engine/Source/Runtime/PakFile/Private/IPlatformFilePak.cpp
+++ b/Engine/Source/Runtime/PakFile/Private/IPlatformFilePak.cpp
@@ -454,6 +454,11 @@ public:
 			PlatformFile.HandleMountCommand(Cmd, Ar);
 			return true;
 		}
+		if (FParse::Command(&Cmd, TEXT("Unmount")))
+		{
+			PlatformFile.HandleUnmountCommand(Cmd, Ar);
+			return true;
+		}
 		else if (FParse::Command(&Cmd, TEXT("PakList")))
 		{
 			PlatformFile.HandlePakListCommand(Cmd, Ar);
@@ -474,6 +479,15 @@ void FPakPlatformFile::HandleMountCommand(const TCHAR* Cmd, FOutputDevice& Ar)
 	}
 }
 
+void FPakPlatformFile::HandleUnmountCommand(const TCHAR* Cmd, FOutputDevice& Ar)
+{
+	const FString PakFilename = FParse::Token(Cmd, false);
+	if (!PakFilename.IsEmpty())
+	{
+		Unmount(*PakFilename);
+	}
+}
+
 void FPakPlatformFile::HandlePakListCommand(const TCHAR* Cmd, FOutputDevice& Ar)
 {
 	TArray<FPakListEntry> Paks;
@@ -494,6 +508,7 @@ FPakPlatformFile::FPakPlatformFile()
 FPakPlatformFile::~FPakPlatformFile()
 {
 	FCoreDelegates::OnMountPak.Unbind();
+	FCoreDelegates::OnUnmountPak.Unbind();
 
 	// We need to flush async IO... if it hasn't been shut down already.
 	if (FIOSystem::HasShutdown() == false)
@@ -671,6 +686,7 @@ bool FPakPlatformFile::Initialize(IPlatformFile* Inner, const TCHAR* CmdLine)
 #endif // !UE_BUILD_SHIPPING
 
 	FCoreDelegates::OnMountPak.BindRaw(this, &FPakPlatformFile::HandleMountPakDelegate);
+	FCoreDelegates::OnUnmountPak.BindRaw(this, &FPakPlatformFile::HandleUnmountPakDelegate);
 	return !!LowerLevel;
 }
 
@@ -710,6 +726,25 @@ bool FPakPlatformFile::Mount(const TCHAR* InPakFilename, uint32 PakOrder, const
 	return bSuccess;
 }
 
+bool FPakPlatformFile::Unmount(const TCHAR* InPakFilename)
+{
+	{
+		FScopeLock ScopedLock(&PakListCritical); 
+
+		for (int32 PakIndex = 0; PakIndex < PakFiles.Num(); PakIndex++)
+		{
+			if (PakFiles[PakIndex].PakFile->GetFilename() == InPakFilename)
+			{
+				delete PakFiles[PakIndex].PakFile;
+				PakFiles.RemoveAt(PakIndex);
+				return true;
+			}
+		}
+	}
+
+	return false;
+}
+
 IFileHandle* FPakPlatformFile::CreatePakFileHandle(const TCHAR* Filename, FPakFile* PakFile, const FPakEntry* FileEntry)
 {
 	IFileHandle* Result = NULL;
@@ -744,6 +779,11 @@ bool FPakPlatformFile::HandleMountPakDelegate(const FString& PakFilePath, uint32
 	return Mount(*PakFilePath, PakOrder);
 }
 
+bool FPakPlatformFile::HandleUnmountPakDelegate(const FString& PakFilePath)
+{
+	return Unmount(*PakFilePath);
+}
+
 IFileHandle* FPakPlatformFile::OpenRead(const TCHAR* Filename)
 {
 	IFileHandle* Result = NULL;
diff --git a/Engine/Source/Runtime/PakFile/Public/IPlatformFilePak.h b/Engine/Source/Runtime/PakFile/Public/IPlatformFilePak.h
--- a/Engine/Source/Runtime/PakFile/Public/IPlatformFilePak.h
+++ b/Engine/Source/Runtime/PakFile/Public/IPlatformFilePak.h
@@ -561,7 +561,7 @@ public:
 	{
 		return Timestamp;
 	}
-
+	
 private:
 
 	/**
@@ -867,6 +867,11 @@ class PAKFILE_API FPakPlatformFile : public IPlatformFile
 	 */
 	bool HandleMountPakDelegate(const FString& PakFilePath, uint32 PakOrder);
 
+	/**
+	 * Handler for device delegate to prompt us to unload a pak.
+	 */
+	bool HandleUnmountPakDelegate(const FString& PakFilePath);
+
 	/**
 	 * Finds all pak files in the given directory.
 	 *
@@ -932,6 +937,8 @@ public:
 	 */
 	bool Mount(const TCHAR* InPakFilename, uint32 PakOrder, const TCHAR* InPath = NULL);
 
+	bool Unmount(const TCHAR* InPakFilename);
+
 	/**
 	 * Finds a file in the specified pak files.
 	 *
@@ -1321,6 +1328,7 @@ public:
 #if !UE_BUILD_SHIPPING
 	void HandlePakListCommand(const TCHAR* Cmd, FOutputDevice& Ar);
 	void HandleMountCommand(const TCHAR* Cmd, FOutputDevice& Ar);
+	void HandleUnmountCommand(const TCHAR* Cmd, FOutputDevice& Ar);
 #endif
 	// END Console commands
 };
