Temporary fix for failed check(TlsSlot) when using compressed pak files.

diff --git a/Engine/Source/Runtime/PakFile/Private/IPlatformFilePak.cpp b/Engine/Source/Runtime/PakFile/Private/IPlatformFilePak.cpp
--- a/Engine/Source/Runtime/PakFile/Private/IPlatformFilePak.cpp
+++ b/Engine/Source/Runtime/PakFile/Private/IPlatformFilePak.cpp
@@ -104,7 +104,9 @@ public:
 
 		FORCEINLINE TStatId GetStatId() const
 		{
-			RETURN_QUICK_DECLARE_CYCLE_STAT(FPakUncompressTask, STATGROUP_ThreadPoolAsyncTasks);
+			// TODO: This is called too early in engine startup.
+			return TStatId();
+			//RETURN_QUICK_DECLARE_CYCLE_STAT(FPakUncompressTask, STATGROUP_ThreadPoolAsyncTasks);
 		}
 	};
 
