Don't crash if the pak file has no files in it.

diff --git a/Engine/Source/Runtime/PakFile/Public/IPlatformFilePak.h b/Engine/Source/Runtime/PakFile/Public/IPlatformFilePak.h
--- a/Engine/Source/Runtime/PakFile/Public/IPlatformFilePak.h
+++ b/Engine/Source/Runtime/PakFile/Public/IPlatformFilePak.h
@@ -504,8 +504,9 @@ public:
 		FFileIterator(const FPakFile& InPakFile)
 		:	PakFile(InPakFile)
 		, IndexIt(PakFile.GetIndex())
-		, DirectoryIt(IndexIt.Value())
-		{}
+		, DirectoryIt((IndexIt ? IndexIt.Value() : FPakDirectory()))
+		{
+		}
 
 		FFileIterator& operator++()		
 		{ 
