Fixed default directory iterator assignment in FFileIterator.

diff --git a/Engine/Source/Runtime/PakFile/Public/IPlatformFilePak.h b/Engine/Source/Runtime/PakFile/Public/IPlatformFilePak.h
--- a/Engine/Source/Runtime/PakFile/Public/IPlatformFilePak.h
+++ b/Engine/Source/Runtime/PakFile/Public/IPlatformFilePak.h
@@ -504,7 +504,7 @@ public:
 		FFileIterator(const FPakFile& InPakFile)
 		:	PakFile(InPakFile)
 		, IndexIt(PakFile.GetIndex())
-		, DirectoryIt((IndexIt ? IndexIt.Value() : FPakDirectory()))
+		, DirectoryIt((IndexIt ? IndexIt.Value() : FPakDirectory().CreateConstIterator()))
 		{
 		}
 
