Another revision to IPlatformPak::FFileIterator constructor. (Fixed initialisation of DirectoryIt)

diff --git a/Engine/Source/Runtime/PakFile/Public/IPlatformFilePak.h b/Engine/Source/Runtime/PakFile/Public/IPlatformFilePak.h
--- a/Engine/Source/Runtime/PakFile/Public/IPlatformFilePak.h
+++ b/Engine/Source/Runtime/PakFile/Public/IPlatformFilePak.h
@@ -504,7 +504,7 @@ public:
 		FFileIterator(const FPakFile& InPakFile)
 		:	PakFile(InPakFile)
 		, IndexIt(PakFile.GetIndex())
-		, DirectoryIt((IndexIt ? IndexIt.Value() : FPakDirectory().CreateConstIterator()))
+		, DirectoryIt((IndexIt ? FPakDirectory::TConstIterator(IndexIt.Value()): FPakDirectory()))
 		{
 		}
 
