Merging 2496325

//depot/UE4-UT/...

to //depot/UE4/...

diff --git a/Engine/Source/Runtime/PakFile/Private/IPlatformFilePak.cpp b/Engine/Source/Runtime/PakFile/Private/IPlatformFilePak.cpp
--- a/Engine/Source/Runtime/PakFile/Private/IPlatformFilePak.cpp
+++ b/Engine/Source/Runtime/PakFile/Private/IPlatformFilePak.cpp
@@ -680,7 +680,27 @@ bool FPakPlatformFile::Initialize(IPlatformFile* Inner, const TCHAR* CmdLine)
 		}
 		if (bLoadPak)
 		{
-			Mount(*PakFilename, 0);
+			// hardcode default load ordering of game main pak -> game content -> engine content -> saved dir
+			// would be better to make this config but not even the config system is initialized here so we can't do that
+			uint32 PakOrder = 0;
+			if (PakFilename.StartsWith(FString::Printf(TEXT("%sPaks/%s-"), *FPaths::GameContentDir(), FApp::GetGameName())))
+			{
+				PakOrder = 4;
+			}
+			else if (PakFilename.StartsWith(FPaths::GameContentDir()))
+			{
+				PakOrder = 3;
+			}
+			else if (PakFilename.StartsWith(FPaths::EngineContentDir()))
+			{
+				PakOrder = 2;
+			}
+			else if (PakFilename.StartsWith(FPaths::GameSavedDir()))
+			{
+				PakOrder = 1;
+			}
+
+			Mount(*PakFilename, PakOrder);
 		}
 	}
 
