Forced patch pak files to be mounted on top of original pak files (first to be read in the queue).
This will need to be changed when the manifests for installed pak files are implemented.

diff --git a/Engine/Source/Runtime/PakFile/Private/IPlatformFilePak.cpp b/Engine/Source/Runtime/PakFile/Private/IPlatformFilePak.cpp
--- a/Engine/Source/Runtime/PakFile/Private/IPlatformFilePak.cpp
+++ b/Engine/Source/Runtime/PakFile/Private/IPlatformFilePak.cpp
@@ -726,6 +726,11 @@ bool FPakPlatformFile::Mount(const TCHAR* InPakFilename, uint32 PakOrder, const
 			{
 				Pak->SetMountPoint(InPath);
 			}
+			FString PakFilename = InPakFilename;
+			if ( PakFilename.EndsWith(TEXT("_P.pak")) )
+			{
+				PakOrder += 100;
+			}
 			{
 				// Add new pak file
 				FScopeLock ScopedLock(&PakListCritical);
