Unshelved from pending changelist '2618083':

PR #1377: use StableSort to sort pak files to preserve previous sorting by fileà (Contributed by [REDACTED_PERSON])

diff --git a/Engine/Source/Runtime/PakFile/Private/IPlatformFilePak.cpp b/Engine/Source/Runtime/PakFile/Private/IPlatformFilePak.cpp
--- a/Engine/Source/Runtime/PakFile/Private/IPlatformFilePak.cpp
+++ b/Engine/Source/Runtime/PakFile/Private/IPlatformFilePak.cpp
@@ -738,7 +738,7 @@ bool FPakPlatformFile::Mount(const TCHAR* InPakFilename, uint32 PakOrder, const
 				Entry.ReadOrder = PakOrder;
 				Entry.PakFile = Pak;
 				PakFiles.Add(Entry);
-				PakFiles.Sort();
+				PakFiles.StableSort();
 			}
 			bSuccess = true;
 		}
