Fix for UE-29057 - Crash occuring in serialization when signed paks are used with compressed packages

diff --git a/Engine/Source/Runtime/PakFile/Private/SignedArchiveReader.cpp b/Engine/Source/Runtime/PakFile/Private/SignedArchiveReader.cpp
--- a/Engine/Source/Runtime/PakFile/Private/SignedArchiveReader.cpp
+++ b/Engine/Source/Runtime/PakFile/Private/SignedArchiveReader.cpp
@@ -323,6 +323,7 @@ FSignedArchiveReader::FSignedArchiveReader(FArchive* InPakReader, FChunkCacheWor
 	, SignatureChecker(InSignatureChecker)
 {
 	// Cache global info about the archive
+	ArIsLoading = true;
 	SizeOnDisk = PakReader->TotalSize();
 	ChunkCount = SizeOnDisk / FPakInfo::MaxChunkDataSize + ((SizeOnDisk % FPakInfo::MaxChunkDataSize) ? 1 : 0);
 	PakSize = SizeOnDisk;
