Copying //UE4/Dev-Core to //UE4/Dev-Main (Source: //UE4/Dev-Core @ 3151653)

diff --git a/Engine/Source/Runtime/PakFile/Private/IPlatformFilePak.cpp b/Engine/Source/Runtime/PakFile/Private/IPlatformFilePak.cpp
--- a/Engine/Source/Runtime/PakFile/Private/IPlatformFilePak.cpp
+++ b/Engine/Source/Runtime/PakFile/Private/IPlatformFilePak.cpp
@@ -19,13 +19,17 @@ DEFINE_LOG_CATEGORY(LogPakFile);
 DEFINE_STAT(STAT_PakFile_Read);
 DEFINE_STAT(STAT_PakFile_NumOpenHandles);
 
+#ifndef EXCLUDE_NONPAK_UE_EXTENSIONS
+#define EXCLUDE_NONPAK_UE_EXTENSIONS USE_NEW_ASYNC_IO	// Use .Build.cs file to disable this if the game relies on accessing loose files
+#endif
+
 FFilenameSecurityDelegate& FPakPlatformFile::GetFilenameSecurityDelegate()
 {
 	static FFilenameSecurityDelegate Delegate;
 	return Delegate;
 }
 
-#if USE_NEW_ASYNC_IO
+#if USE_NEW_ASYNC_IO && !IS_PROGRAM && !WITH_EDITOR
 	#define USE_PAK_PRECACHE (1) // you can turn this off to use the async IO stuff without the precache
 #else
 	#define USE_PAK_PRECACHE (0)
@@ -40,11 +44,12 @@ FFilenameSecurityDelegate& FPakPlatformFile::GetFilenameSecurityDelegate()
 
 
 #define PAK_CACHE_GRANULARITY (64 * 1024)
-#define PAK_CACHE_MAX_REQUEST (1024 * 1024)
 #define PAK_CACHE_MAX_REQUESTS (8)
 #define PAK_CACHE_MAX_PRIORITY_DIFFERENCE_MERGE (AIOP_Normal - AIOP_Precache)
+#define PAK_EXTRA_CHECKS DO_CHECK
 
-DECLARE_MEMORY_STAT(TEXT("PakCache Total"), STAT_PakCacheMem, STATGROUP_Memory);
+DECLARE_MEMORY_STAT(TEXT("PakCache Current"), STAT_PakCacheMem, STATGROUP_Memory);
+DECLARE_MEMORY_STAT(TEXT("PakCache High Water"), STAT_PakCacheHighWater, STATGROUP_Memory);
 
 static int32 GPakCache_Enable = 1;
 static FAutoConsoleVariableRef CVar_Enable(
@@ -60,21 +65,15 @@ static FAutoConsoleVariableRef CVar_MaxRequestsToLowerLevel(
 	TEXT("Controls the maximum number of IO requests submitted to the OS filesystem at one time. Limited by PAK_CACHE_MAX_REQUESTS.")
 	);
 
-static int32 GPakCache_MaxRequestSizeToLowerLevelKB = PAK_CACHE_MAX_REQUEST / 1024;
+static int32 GPakCache_MaxRequestSizeToLowerLevelKB = 1024;
 static FAutoConsoleVariableRef CVar_MaxRequestSizeToLowerLevelKB(
 	TEXT("pakcache.MaxRequestSizeToLowerLevellKB"),
 	GPakCache_MaxRequestSizeToLowerLevelKB,
-	TEXT("Controls the maximum size (in KB) of IO requests submitted to the OS filesystem. Limited by PAK_CACHE_MAX_REQUEST.")
+	TEXT("Controls the maximum size (in KB) of IO requests submitted to the OS filesystem.")
 	);
 
-static int32 GPakCache_MaxPrecacheMemMB = 0;
-static FAutoConsoleVariableRef CVar_MaxPrecacheMemMB(
-	TEXT("pakcache.MaxPrecacheMemMB"),
-	GPakCache_MaxPrecacheMemMB,
-	TEXT("Controls the maximum amount of memory (in MB) that will be used for precaching. Total memory use can be higher because not all requests are precache requests. Setting this to zero means unlimited memory.")
-	);
-
-
+//@todoio this won't work with UT etc, need *_API or a real api to work with the pak precacher.
+bool GPakCache_AcceptPrecacheRequests = true;
 
 class FPakPrecacher;
 
@@ -147,13 +146,13 @@ public:
 	FORCEINLINE TItem& Get(TIntervalTreeIndex InIndex)
 	{
 		TIntervalTreeIndex Index = InIndex & ~SaltMask;
-		check((InIndex & SaltMask) == Salt && Index != IntervalTreeInvalidIndex && Index >= 0 && Index < Items.Num() && !FreeItems.Contains(Index));
+		check((InIndex & SaltMask) == Salt && Index != IntervalTreeInvalidIndex && Index >= 0 && Index < (uint32)Items.Num()); //&& !FreeItems.Contains(Index));
 		return Items[Index];
 	}
 	FORCEINLINE void Free(TIntervalTreeIndex InIndex)
 	{
 		TIntervalTreeIndex Index = InIndex & ~SaltMask;
-		check((InIndex & SaltMask) == Salt && Index != IntervalTreeInvalidIndex && Index >= 0 && Index < Items.Num() && !FreeItems.Contains(Index));
+		check((InIndex & SaltMask) == Salt && Index != IntervalTreeInvalidIndex && Index >= 0 && Index < (uint32)Items.Num()); //&& !FreeItems.Contains(Index));
 		Items[Index].~TItem();
 		FreeItems.Push(Index);
 		if (FreeItems.Num() + 1 == Items.Num())
@@ -167,7 +166,7 @@ public:
 	FORCEINLINE void CheckIndex(TIntervalTreeIndex InIndex)
 	{
 		TIntervalTreeIndex Index = InIndex & ~SaltMask;
-		check((InIndex & SaltMask) == Salt && Index != IntervalTreeInvalidIndex && Index >= 0 && Index < Items.Num() && !FreeItems.Contains(Index));
+		check((InIndex & SaltMask) == Salt && Index != IntervalTreeInvalidIndex && Index >= 0 && Index < (uint32)Items.Num()); // && !FreeItems.Contains(Index));
 	}
 };
 
@@ -916,11 +915,34 @@ public:
 	{
 		if (!PakPrecacherSingleton)
 		{
-			verify(!FPlatformAtomics::InterlockedCompareExchangePointer((void**)&PakPrecacherSingleton, new FPakPrecacher(InLowerLevel), NULL));
+			verify(!FPlatformAtomics::InterlockedCompareExchangePointer((void**)&PakPrecacherSingleton, new FPakPrecacher(InLowerLevel), nullptr));
 		}
 		check(PakPrecacherSingleton);
 	}
 
+	static void Shutdown()
+	{
+		if (PakPrecacherSingleton)
+		{
+			FPakPrecacher* LocalPakPrecacherSingleton = PakPrecacherSingleton;
+			if (LocalPakPrecacherSingleton && LocalPakPrecacherSingleton == FPlatformAtomics::InterlockedCompareExchangePointer((void**)&PakPrecacherSingleton, nullptr, LocalPakPrecacherSingleton))
+			{
+				double StartTime = FPlatformTime::Seconds();
+				while (!LocalPakPrecacherSingleton->IsProbablyIdle())
+				{
+					FPlatformProcess::SleepNoStats(0.001f);
+					if (FPlatformTime::Seconds() - StartTime > 10.0)
+					{
+						UE_LOG(LogPakFile, Error, TEXT("FPakPrecacher was not idle after 10s, exiting anyway and leaking."));
+						return;
+					}
+				}
+				delete PakPrecacherSingleton;
+			}
+		}
+		check(!PakPrecacherSingleton);
+	}
+
 	static FPakPrecacher& Get()
 	{
 		check(PakPrecacherSingleton);
@@ -938,15 +960,8 @@ public:
 		, Frees(0)
 		, LoadSize(0)
 	{
-		check(LowerLevel);
-		if (FPlatformProcess::SupportsMultithreading())
-		{
-			GPakCache_MaxRequestsToLowerLevel = FMath::Max(FMath::Min(FPlatformMisc::NumberOfIOWorkerThreadsToSpawn(), GPakCache_MaxRequestsToLowerLevel), 1);
-		}
-		else
-		{
-			GPakCache_MaxRequestsToLowerLevel = 0;
-		}
+		check(LowerLevel && FPlatformProcess::SupportsMultithreading());
+		GPakCache_MaxRequestsToLowerLevel = FMath::Max(FMath::Min(FPlatformMisc::NumberOfIOWorkerThreadsToSpawn(), GPakCache_MaxRequestsToLowerLevel), 1);
 		check(GPakCache_MaxRequestsToLowerLevel <= PAK_CACHE_MAX_REQUESTS);
 		CallbackFromLower = 
 			[this](bool bWasCanceled, IAsyncReadRequest* Request)
@@ -961,14 +976,9 @@ public:
 		return LowerLevel;
 	}
 
-	int64 MaxMemPrecache()
-	{
-		return GPakCache_MaxPrecacheMemMB * 1024ll * 1024ll;
-	}
-
 	bool HasEnoughRoomForPrecache()
 	{
-		return GPakCache_MaxPrecacheMemMB == 0 || BlockMemory < MaxMemPrecache();
+		return GPakCache_AcceptPrecacheRequests;
 	}
 
 private: // below here we assume CachedFilesScopeLock until we get to the next section
@@ -1046,7 +1056,7 @@ private: // below here we assume CachedFilesScopeLock until we get to the next s
 				&InFlightOrDone[0]
 				);
 		}
-		for (int32 Index = 0; Index < NumQWords; Index++)
+		for (uint32 Index = 0; Index < NumQWords; Index++)
 		{
 			if (InFlightOrDone[Index] != MAX_uint64)
 			{
@@ -1105,7 +1115,7 @@ private: // below here we assume CachedFilesScopeLock until we get to the next s
 				return true;
 			}
 			);
-			for (int32 Index = 0; Index < NumQWords; Index++)
+			for (uint32 Index = 0; Index < NumQWords; Index++)
 			{
 				if (InFlightOrDone[Index] != MAX_uint64)
 				{
@@ -1137,7 +1147,7 @@ private: // below here we assume CachedFilesScopeLock until we get to the next s
 				}
 				);
 
-				for (int32 Index = 0; Index < NumQWords; Index++)
+				for (uint32 Index = 0; Index < NumQWords; Index++)
 				{
 					if (InFlightOrDone[Index] != MAX_uint64)
 					{
@@ -1149,7 +1159,7 @@ private: // below here we assume CachedFilesScopeLock until we get to the next s
 		}
 		else
 		{
-#if DO_CHECK
+#if PAK_EXTRA_CHECKS
 			OverlappingNodesInIntervalTree<FCacheBlock>(
 				Pak.CacheBlocks[(int32)EBlockStatus::InFlight],
 				CacheBlockAllocator,
@@ -1191,16 +1201,13 @@ private: // below here we assume CachedFilesScopeLock until we get to the next s
 
 	void ClearBlock(FCacheBlock &Block)
 	{
-#if 0
-		if (GetRequestPakIndex(Block.OffsetAndPakIndex) == 1 && GetRequestOffset(Block.OffsetAndPakIndex) <= 0x22f573d3ll && GetRequestOffset(Block.OffsetAndPakIndex) + Block.Size >= 0x22f573d3ll)
-		{
-			UE_LOG(LogTemp, Error, TEXT("Hi"));
-		}
-#endif
 		if (Block.Memory)
 		{
+			check(Block.Size);
 			BlockMemory -= Block.Size;
 			DEC_MEMORY_STAT_BY(STAT_PakCacheMem, Block.Size);
+			check(BlockMemory >= 0);
+
 			FMemory::Free(Block.Memory);
 			Block.Memory = nullptr;
 		}
@@ -1318,9 +1325,10 @@ private: // below here we assume CachedFilesScopeLock until we get to the next s
 		FJoinedOffsetAndPakIndex BestNext = MAX_uint64;
 
 		OutPriority = AIOP_MIN;
+		bool bAnyOutstanding = false;
 		for (EAsyncIOPriority Priority = AIOP_MAX;; Priority = EAsyncIOPriority(int32(Priority) - 1))
 		{
-			if (Priority == AIOP_Precache && !bAcceptingPrecacheRequests)
+			if (Priority == AIOP_Precache && !bAcceptingPrecacheRequests && bAnyOutstanding)
 			{
 				break;
 			}
@@ -1336,6 +1344,10 @@ private: // below here we assume CachedFilesScopeLock until we get to the next s
 				for (; BestNext == MAX_uint64 && PakIndex < CachedPakData.Num(); PakIndex++)
 				{
 					FPakData& Pak = CachedPakData[PakIndex];
+					if (Pak.InRequests[Priority][(int32)EInRequestStatus::Complete] != IntervalTreeInvalidIndex)
+					{
+						bAnyOutstanding = true;
+					}
 					if (Pak.InRequests[Priority][(int32)EInRequestStatus::Waiting] != IntervalTreeInvalidIndex)
 					{
 						uint64 Limit = uint64(Pak.TotalSize - 1);
@@ -1397,7 +1409,7 @@ private: // below here we assume CachedFilesScopeLock until we get to the next s
 		FPakData& Pak = CachedPakData[PakIndex];
 		check(Offset < Pak.TotalSize);
 		int64 FirstByte = AlignDown(Offset, PAK_CACHE_GRANULARITY);
-		int64 LastByte = FMath::Min(Align(FirstByte + PAK_CACHE_MAX_REQUEST, PAK_CACHE_GRANULARITY) - 1, Pak.TotalSize - 1);
+		int64 LastByte = FMath::Min(Align(FirstByte + (GPakCache_MaxRequestSizeToLowerLevelKB * 1024), PAK_CACHE_GRANULARITY) - 1, Pak.TotalSize - 1);
 		check(FirstByte >= 0 && LastByte < Pak.TotalSize && LastByte >= 0 && LastByte >= FirstByte);
 
 		uint32 NumBits = (PAK_CACHE_GRANULARITY + LastByte - FirstByte) / PAK_CACHE_GRANULARITY;
@@ -1475,7 +1487,7 @@ private: // below here we assume CachedFilesScopeLock until we get to the next s
 
 
 		int64 Size = PAK_CACHE_GRANULARITY * 64 * NumQWords;
-		for (int32 Index = 0; Index < NumQWords; Index++)
+		for (uint32 Index = 0; Index < NumQWords; Index++)
 		{
 			uint64 NotAlreadyInFlightAndRequested = ((~InFlightOrDone[Index]) & Requested[Index]);
 			if (NotAlreadyInFlightAndRequested != MAX_uint64)
@@ -1489,7 +1501,7 @@ private: // below here we assume CachedFilesScopeLock until we get to the next s
 				break;
 			}
 		}
-		check(Size > 0 && Size <= PAK_CACHE_MAX_REQUEST);
+		check(Size > 0 && Size <= (GPakCache_MaxRequestSizeToLowerLevelKB * 1024));
 		Size = FMath::Min(FirstByte + Size, LastByte + 1) - FirstByte;
 
 		TIntervalTreeIndex NewIndex = CacheBlockAllocator.Alloc();
@@ -1538,7 +1550,7 @@ private: // below here we assume CachedFilesScopeLock until we get to the next s
 				}
 				);
 			}
-#if DO_CHECK
+#if PAK_EXTRA_CHECKS
 			OverlappingNodesInIntervalTree<FPakInRequest>(
 				Pak.InRequests[Priority][(int32)EInRequestStatus::InFlight],
 				InRequestAllocator,
@@ -1686,19 +1698,25 @@ private: // below here we assume CachedFilesScopeLock until we get to the next s
 		else
 		{
 			Block.Memory = Request->GetReadResults();
+			check(Block.Memory && Block.Size);
 			BlockMemory += Block.Size;
+			check(BlockMemory > 0);
+			DEC_MEMORY_STAT_BY(STAT_AsyncFileMemory, Block.Size);
 			INC_MEMORY_STAT_BY(STAT_PakCacheMem, Block.Size);
 
 			if (BlockMemory > BlockMemoryHighWater)
 			{
 				BlockMemoryHighWater = BlockMemory;
+				SET_MEMORY_STAT(STAT_PakCacheHighWater, BlockMemoryHighWater);
 
+#if 1
 				static int64 LastPrint = 0;
 				if (BlockMemoryHighWater / 1024 / 1024 != LastPrint)
 				{
 					LastPrint = BlockMemoryHighWater / 1024 / 1024;
 					FPlatformMisc::LowLevelOutputDebugStringf(TEXT("Precache HighWater %dMB\r\n"), int32(LastPrint));
 				}
+#endif
 			}
 			Block.Status = EBlockStatus::Complete;
 			AddToIntervalTree<FCacheBlock>(
@@ -1758,7 +1776,7 @@ private: // below here we assume CachedFilesScopeLock until we get to the next s
 		return false;
 	}
 
-	uint8* GetCompletedRequestData(FPakInRequest& DoneRequest)
+	bool GetCompletedRequestData(FPakInRequest& DoneRequest, uint8* Result)
 	{
 		// CachedFilesScopeLock is locked
 		check(DoneRequest.Status == EInRequestStatus::Complete);
@@ -1769,7 +1787,6 @@ private: // below here we assume CachedFilesScopeLock until we get to the next s
 		FPakData& Pak = CachedPakData[PakIndex];
 		check(Offset + DoneRequest.Size <= Pak.TotalSize && DoneRequest.Size > 0 && DoneRequest.Priority >= AIOP_MIN && DoneRequest.Priority <= AIOP_MAX && DoneRequest.Status == EInRequestStatus::Complete);
 
-		uint8* Result = (uint8*)FMemory::Malloc(Size);
 		int64 BytesCopied = 0;
 
 #if 0 // this path removes the block in one pass, however, this is not what we want because it wrecks precaching, if we change back GetCompletedRequest needs to maybe start a new request and the logic of the IAsyncFile read needs to change
@@ -1836,7 +1853,7 @@ private: // below here we assume CachedFilesScopeLock until we get to the next s
 		check(BytesCopied == Size);
 
 
-		return Result;
+		return true;
 	}
 
 	///// Below here are the thread entrypoints
@@ -1881,10 +1898,11 @@ public:
 			}
 			CachedPakData.Add(FPakData(Handle, File, PakFileSize));
 			PakIndexPtr = &CachedPaks.Add(File, CachedPakData.Num() - 1);
+			UE_LOG(LogPakFile, Log, TEXT("New pak file %s added to pak precacher."), *File.ToString());
 		}
 		uint16 PakIndex = *PakIndexPtr;
 		FPakData& Pak = CachedPakData[PakIndex];
-		check(Pak.Name == File && Pak.TotalSize == PakFileSize);
+		check(Pak.Name == File && Pak.TotalSize == PakFileSize && Pak.Handle);
 
 		TIntervalTreeIndex RequestIndex = InRequestAllocator.Alloc();
 		FPakInRequest& Request = InRequestAllocator.Get(RequestIndex);
@@ -1914,7 +1932,7 @@ public:
 		return true;
 	}
 
-	uint8* GetCompletedRequest(IPakRequestor* Owner)
+	bool GetCompletedRequest(IPakRequestor* Owner, uint8* UserSuppliedMemory)
 	{
 		check(Owner);
 		FScopeLock Lock(&CachedFilesScopeLock);
@@ -1925,10 +1943,9 @@ public:
 		{
 			FPakInRequest& Request = InRequestAllocator.Get(RequestIndex);
 			check(Owner == Request.Owner && Request.Status == EInRequestStatus::Complete && Request.UniqueID == Request.Owner->UniqueID && RequestIndex == Request.Owner->InRequestIndex &&  Request.OffsetAndPakIndex == Request.Owner->OffsetAndPakIndex);
-			uint8* Result = GetCompletedRequestData(Request);
-			return Result;
+			return GetCompletedRequestData(Request, UserSuppliedMemory);
 		}
-		return nullptr; // canceled
+		return false; // canceled
 	}
 
 	void CancelRequest(IPakRequestor* Owner)
@@ -1940,16 +1957,11 @@ public:
 		static_assert(IntervalTreeInvalidIndex == 0, "FindRef will return 0 for something not found");
 		if (RequestIndex)
 		{
-			bool bWasOver = !HasEnoughRoomForPrecache();
 			FPakInRequest& Request = InRequestAllocator.Get(RequestIndex);
 			check(Owner == Request.Owner && Request.UniqueID == Request.Owner->UniqueID && RequestIndex == Request.Owner->InRequestIndex &&  Request.OffsetAndPakIndex == Request.Owner->OffsetAndPakIndex);
 			RemoveRequest(RequestIndex);
-			bool bNowOver = !HasEnoughRoomForPrecache();
-			if (bWasOver && !bNowOver)
-			{
-				StartNextRequest();
-			}
 		}
+		StartNextRequest();
 	}
 
 	bool IsProbablyIdle() // nothing to prevent new requests from being made before I return
@@ -1958,6 +1970,139 @@ public:
 		return !HasRequestsAtStatus(EInRequestStatus::Waiting) && !HasRequestsAtStatus(EInRequestStatus::InFlight);
 	}
 
+	void Unmount(FName PakFile)
+	{
+		FScopeLock Lock(&CachedFilesScopeLock);
+		uint16* PakIndexPtr = CachedPaks.Find(PakFile);
+		if (!PakIndexPtr)
+		{
+			UE_LOG(LogPakFile, Log, TEXT("Pak file %s was never used, so nothing to unmount"), *PakFile.ToString());
+			return; // never used for anything, nothing to check or clean up
+		}
+		uint16 PakIndex = *PakIndexPtr;
+		FPakData& Pak = CachedPakData[PakIndex];
+		int64 Offset = MakeJoinedRequest(PakIndex, 0);
+
+		bool bHasOutstandingRequests = false;
+
+		OverlappingNodesInIntervalTree<FCacheBlock>(
+			Pak.CacheBlocks[(int32)EBlockStatus::Complete],
+			CacheBlockAllocator,
+			0,
+			Offset + Pak.TotalSize - 1,
+			0,
+			Pak.MaxNode,
+			Pak.StartShift,
+			Pak.MaxShift,
+			[&bHasOutstandingRequests](TIntervalTreeIndex BlockIndex) -> bool
+		{
+			check(!"Pak cannot be unmounted with outstanding requests");
+			bHasOutstandingRequests = true;
+			return false;
+		}
+		);
+		OverlappingNodesInIntervalTree<FCacheBlock>(
+			Pak.CacheBlocks[(int32)EBlockStatus::InFlight],
+			CacheBlockAllocator,
+			0,
+			Offset + Pak.TotalSize - 1,
+			0,
+			Pak.MaxNode,
+			Pak.StartShift,
+			Pak.MaxShift,
+			[&bHasOutstandingRequests](TIntervalTreeIndex BlockIndex) -> bool
+		{
+			check(!"Pak cannot be unmounted with outstanding requests");
+			bHasOutstandingRequests = true;
+			return false;
+		}
+		);
+		for (EAsyncIOPriority Priority = AIOP_MAX;; Priority = EAsyncIOPriority(int32(Priority) - 1))
+		{
+			OverlappingNodesInIntervalTree<FPakInRequest>(
+				Pak.InRequests[Priority][(int32)EInRequestStatus::InFlight],
+				InRequestAllocator,
+				0,
+				Offset + Pak.TotalSize - 1,
+				0,
+				Pak.MaxNode,
+				Pak.StartShift,
+				Pak.MaxShift,
+				[&bHasOutstandingRequests](TIntervalTreeIndex BlockIndex) -> bool
+			{
+				check(!"Pak cannot be unmounted with outstanding requests");
+				bHasOutstandingRequests = true;
+				return false;
+			}
+			);
+			OverlappingNodesInIntervalTree<FPakInRequest>(
+				Pak.InRequests[Priority][(int32)EInRequestStatus::Complete],
+				InRequestAllocator,
+				0,
+				Offset + Pak.TotalSize - 1,
+				0,
+				Pak.MaxNode,
+				Pak.StartShift,
+				Pak.MaxShift,
+				[&bHasOutstandingRequests](TIntervalTreeIndex BlockIndex) -> bool
+			{
+				check(!"Pak cannot be unmounted with outstanding requests");
+				bHasOutstandingRequests = true;
+				return false;
+			}
+			);
+			OverlappingNodesInIntervalTree<FPakInRequest>(
+				Pak.InRequests[Priority][(int32)EInRequestStatus::Waiting],
+				InRequestAllocator,
+				0,
+				Offset + Pak.TotalSize - 1,
+				0,
+				Pak.MaxNode,
+				Pak.StartShift,
+				Pak.MaxShift,
+				[&bHasOutstandingRequests](TIntervalTreeIndex BlockIndex) -> bool
+			{
+				check(!"Pak cannot be unmounted with outstanding requests");
+				bHasOutstandingRequests = true;
+				return false;
+			}
+			);
+			if (Priority == AIOP_MIN)
+			{
+				break;
+			}
+		}
+		if (!bHasOutstandingRequests)
+		{
+			UE_LOG(LogPakFile, Log, TEXT("Pak file %s removed from pak precacher."), *PakFile.ToString());
+			CachedPaks.Remove(PakFile);
+			check(Pak.Handle);
+			delete Pak.Handle;
+			Pak.Handle = nullptr;
+			int32 NumToTrim = 0;
+			for (int32 Index = CachedPakData.Num() - 1; Index >= 0; Index--)
+			{
+				if (!CachedPakData[Index].Handle)
+				{
+					NumToTrim++;
+				}
+				else
+				{
+					break;
+				}
+			}
+			if (NumToTrim)
+			{
+				CachedPakData.RemoveAt(CachedPakData.Num() - NumToTrim, NumToTrim);
+			}
+		}
+		else
+		{
+			UE_LOG(LogPakFile, Log, TEXT("Pak file %s was NOT removed from pak precacher because it had outstanding requests."), *PakFile.ToString());
+		}
+	}
+
+
 	// these are not threadsafe and should only be used for synthetic testing
 	uint64 GetLoadSize()
 	{
@@ -2021,15 +2166,6 @@ static FAutoConsoleCommand WaitPrecacheCmd(
 	FConsoleCommandWithArgsDelegate::CreateStatic(&WaitPrecache)
 	);
 
-static void FlushPrecache(const TArray<FString>& Args)
-{
-}
-
-static FAutoConsoleCommand FlushPrecacheCmd(
-	TEXT("pak.FlushPrecache"),
-	TEXT("Does nothing (yet)."),
-	FConsoleCommandWithArgsDelegate::CreateStatic(&FlushPrecache)
-	);
 static void DumpBlocks(const TArray<FString>& Args)
 {
 	FPakPrecacher::Get().DumpBlocks();
@@ -2043,25 +2179,34 @@ static FAutoConsoleCommand DumpBlocksCmd(
 
 static FCriticalSection FPakReadRequestEvent;
 
+class FPakAsyncReadFileHandle;
+//uncompress(unencrypt(checksig())))
+
 class FPakReadRequest : public IAsyncReadRequest, public IPakRequestor
 {
 	int64 Offset;
 	int64 BytesToRead;
 	FEvent* WaitEvent;
+	int32 BlockIndex;
 	EAsyncIOPriority Priority;
 	bool bRequestOutstanding;
 	bool bNeedsRemoval;
+	bool bInternalRequest; // we are using this internally to deal with compressed, encrypted and signed, so we want the memory back from a precache request.
+
 public:
-	FPakReadRequest(FName InPakFile, int64 PakFileSize, FAsyncFileCallBack* CompleteCallback, int64 InOffset, int64 InBytesToRead, EAsyncIOPriority InPriority)
-		: IAsyncReadRequest(CompleteCallback, false)
+	FPakReadRequest(FName InPakFile, int64 PakFileSize, FAsyncFileCallBack* CompleteCallback, int64 InOffset, int64 InBytesToRead, EAsyncIOPriority InPriority, uint8* UserSuppliedMemory, bool bInInternalRequest = false, int32 InBlockIndex = -1)
+		: IAsyncReadRequest(CompleteCallback, false, UserSuppliedMemory)
 		, Offset(InOffset)
 		, BytesToRead(InBytesToRead)
 		, WaitEvent(nullptr)
+		, BlockIndex(InBlockIndex)
 		, Priority(InPriority)
 		, bRequestOutstanding(true)
 		, bNeedsRemoval(true)
+		, bInternalRequest(bInInternalRequest)
 	{
 		check(Offset >= 0 && BytesToRead > 0);
+		check(bInternalRequest || Priority > AIOP_Precache || !bUserSuppliedMemory); // you never get bits back from a precache request, so why supply memory?
 		if (!FPakPrecacher::Get().QueueRequest(this, InPakFile, PakFileSize, Offset, BytesToRead, Priority))
 		{
 			bRequestOutstanding = false;
@@ -2075,6 +2220,13 @@ public:
 		{
 			FPakPrecacher::Get().CancelRequest(this);
 		}
+		if (Memory && !bUserSuppliedMemory)
+		{
+			// this can happen with a race on cancel, it is ok, they didn't take the memory, free it now
+			DEC_MEMORY_STAT_BY(STAT_AsyncFileMemory, BytesToRead);
+			FMemory::Free(Memory);
+		}
+		Memory = nullptr;
 	}
 
 	// IAsyncReadRequest Interface
@@ -2124,14 +2276,23 @@ public:
 
 	virtual void RequestIsComplete() override
 	{
-		check(bRequestOutstanding && !Memory);
-		if (bCanceled || Priority <= AIOP_Precache)
-		{
-		}
-		else
+		check(bRequestOutstanding);
+		if (!bCanceled && (bInternalRequest || Priority > AIOP_Precache))
 		{
-			Memory = FPakPrecacher::Get().GetCompletedRequest(this);
-			// Memory could be null here if the request was canceled
+			if (!bUserSuppliedMemory)
+			{
+				check(!Memory);
+				Memory = (uint8*)FMemory::Malloc(BytesToRead);
+				INC_MEMORY_STAT_BY(STAT_AsyncFileMemory, BytesToRead);
+			}
+			else
+			{
+				check(Memory);
+			}
+			if (!FPakPrecacher::Get().GetCompletedRequest(this, Memory))
+			{
+				check(bCanceled);
+			}
 		}
 		SetComplete();
 		FScopeLock Lock(&FPakReadRequestEvent);
@@ -2141,13 +2302,17 @@ public:
 			WaitEvent->Trigger();
 		}
 	}
+	int32 GetBlockIndex()
+	{
+		return BlockIndex;
+	}
 };
 
 class FPakSizeRequest : public IAsyncReadRequest
 {
 public:
 	FPakSizeRequest(FAsyncFileCallBack* CompleteCallback, int64 InFileSize)
-		: IAsyncReadRequest(CompleteCallback, true)
+		: IAsyncReadRequest(CompleteCallback, true, nullptr)
 	{
 		Size = InFileSize;
 		SetComplete();
@@ -2161,36 +2326,521 @@ public:
 };
 
 
+struct FCachedAsyncBlock
+{
+	FPakReadRequest* RawRequest;
+	uint8* Raw; // compressed, encrypted and/or signature not checked
+	uint8* Processed; // decompressed, deencrypted and signature checked
+	FGraphEventRef CPUWorkGraphEvent;
+	int32 RawSize;
+	int32 ProcessedSize;
+	int32 RefCount;
+	bool bInFlight;
+	FCachedAsyncBlock()
+		: RawRequest(0)
+		, Raw(nullptr)
+		, Processed(nullptr)
+		, RawSize(0)
+		, ProcessedSize(0)
+		, RefCount(0)
+		, bInFlight(false)
+	{
+	}
+};
+
+
+class FPakProcessedReadRequest : public IAsyncReadRequest
+{
+	FPakAsyncReadFileHandle* Owner;
+	int64 Offset;
+	int64 BytesToRead;
+	FEvent* WaitEvent;
+	FThreadSafeCounter CompleteRace; // this is used to resolve races with natural completion and cancel; there can be only one.
+	EAsyncIOPriority Priority;
+	bool bRequestOutstanding;
+	bool bHasCancelled;
+	bool bHasCompleted;
+
+public:
+	FPakProcessedReadRequest(FPakAsyncReadFileHandle* InOwner, FAsyncFileCallBack* CompleteCallback, int64 InOffset, int64 InBytesToRead, EAsyncIOPriority InPriority, uint8* UserSuppliedMemory)
+		: IAsyncReadRequest(CompleteCallback, false, UserSuppliedMemory)
+		, Owner(InOwner)
+		, Offset(InOffset)
+		, BytesToRead(InBytesToRead)
+		, WaitEvent(nullptr)
+		, Priority(InPriority)
+		, bRequestOutstanding(true)
+		, bHasCancelled(false)
+		, bHasCompleted(false)
+	{
+		check(Offset >= 0 && BytesToRead > 0);
+		check(Priority > AIOP_Precache || !bUserSuppliedMemory); // you never get bits back from a precache request, so why supply memory?
+	}
+
+	virtual ~FPakProcessedReadRequest()
+	{
+		DoneWithRawRequests();
+		if (Memory && !bUserSuppliedMemory)
+		{
+			// this can happen with a race on cancel, it is ok, they didn't take the memory, free it now
+			DEC_MEMORY_STAT_BY(STAT_AsyncFileMemory, BytesToRead);
+			FMemory::Free(Memory);
+		}
+		Memory = nullptr;
+	}
+
+	// IAsyncReadRequest Interface
+
+	virtual void WaitCompletionImpl(float TimeLimitSeconds) override
+	{
+		if (bRequestOutstanding)
+		{
+			{
+				FScopeLock Lock(&FPakReadRequestEvent);
+				if (bRequestOutstanding)
+				{
+					check(!WaitEvent);
+					WaitEvent = FPlatformProcess::GetSynchEventFromPool(true);
+				}
+			}
+			if (WaitEvent)
+			{
+				if (TimeLimitSeconds == 0.0f)
+				{
+					WaitEvent->Wait();
+					check(!bRequestOutstanding);
+				}
+				else
+				{
+					WaitEvent->Wait(TimeLimitSeconds * 1000.0f);
+				}
+				FPlatformProcess::ReturnSynchEventToPool(WaitEvent);
+				WaitEvent = nullptr;
+			}
+		}
+	}
+	virtual void CancelImpl() override
+	{
+		check(!WaitEvent); // you canceled from a different thread that you waited from
+		if (bRequestOutstanding)
+		{
+			CancelRawRequests();
+			bRequestOutstanding = false;
+			SetComplete();
+		}
+		check(!WaitEvent); // you canceled from a different thread that you waited from
+	}
+
+	void RequestIsComplete()
+	{
+		check(bRequestOutstanding);
+		if (!bCanceled && Priority > AIOP_Precache)
+		{
+			GatherResults();
+		}
+		SetComplete();
+		FScopeLock Lock(&FPakReadRequestEvent);
+		bRequestOutstanding = false;
+		if (WaitEvent)
+		{
+			WaitEvent->Trigger();
+		}
+	}
+	void GatherResults();
+	void DoneWithRawRequests();
+	bool CheckCompletion(const FPakEntry& FileEntry, int32 BlockIndex, TArray<FCachedAsyncBlock>& Blocks);
+	void CancelRawRequests();
+};
+
+FAutoConsoleTaskPriority CPrio_AsyncIOCPUWorkTaskPriority(
+	TEXT("TaskGraph.TaskPriorities.AsyncIOCPUWork"),
+	TEXT("Task and thread priority for decompression, decryption and signature checking of async IO from a pak file."),
+	ENamedThreads::BackgroundThreadPriority, // if we have background priority task threads, then use them...
+	ENamedThreads::NormalTaskPriority, // .. at normal task priority
+	ENamedThreads::NormalTaskPriority // if we don't have background threads, then use normal priority threads at normal task priority instead
+	);
+
+class FAsyncIOCPUWorkTask
+{
+	FPakAsyncReadFileHandle& Owner;
+	int32 BlockIndex;
+
+public:
+	FORCEINLINE FAsyncIOCPUWorkTask(FPakAsyncReadFileHandle& InOwner, int32 InBlockIndex)
+		: Owner(InOwner)
+		, BlockIndex(InBlockIndex)
+	{
+	}
+	static FORCEINLINE TStatId GetStatId()
+	{
+		RETURN_QUICK_DECLARE_CYCLE_STAT(FsyncIOCPUWorkTask, STATGROUP_TaskGraphTasks);
+	}
+	static FORCEINLINE ENamedThreads::Type GetDesiredThread()
+	{
+		return CPrio_AsyncIOCPUWorkTaskPriority.Get();
+	}
+	FORCEINLINE static ESubsequentsMode::Type GetSubsequentsMode()
+	{
+		return ESubsequentsMode::TrackSubsequents;
+	}
+	void DoTask(ENamedThreads::Type CurrentThread, const FGraphEventRef& MyCompletionGraphEvent);
+};
+
 class FPakAsyncReadFileHandle final : public IAsyncReadFileHandle
 {
 	FName PakFile;
 	int64 PakFileSize;
 	int64 OffsetInPak;
-	int64 FileSize;
+	int64 CompressedFileSize;
+	int64 UncompressedFileSize;
+	const FPakEntry* FileEntry;
+	TSet<FPakProcessedReadRequest*> LiveRequests;
+	TArray<FCachedAsyncBlock> Blocks;
+	FAsyncFileCallBack ReadCallbackFunction;
+	FCriticalSection CriticalSection;
+	int32 NumLiveRawRequests;
+
 public:
-	FPakAsyncReadFileHandle(FName InPakFile, int64 InPakFileSize, int64 InOffsetInPak, int64 InFileSize)
-		: PakFile(InPakFile)
-		, PakFileSize(InPakFileSize)
-		, OffsetInPak(InOffsetInPak)
-		, FileSize(InFileSize)
+	FPakAsyncReadFileHandle(const FPakEntry* InFileEntry, FPakFile* InPakFile, const TCHAR* Filename)
+		: PakFile(InPakFile->GetFilenameName())
+		, PakFileSize(InPakFile->TotalSize())
+		, FileEntry(InFileEntry)
+		, NumLiveRawRequests(0)
 	{
-		check(PakFileSize > 0 && OffsetInPak + FileSize <= PakFileSize && OffsetInPak >= 0);
+		OffsetInPak = FileEntry->Offset + FileEntry->GetSerializedSize(InPakFile->GetInfo().Version);
+		UncompressedFileSize = FileEntry->UncompressedSize;
+		CompressedFileSize = FileEntry->UncompressedSize;
+		if (FileEntry->CompressionMethod != COMPRESS_None && UncompressedFileSize)
+		{
+			check(FileEntry->CompressionBlocks.Num());
+			CompressedFileSize = FileEntry->CompressionBlocks.Last().CompressedEnd - OffsetInPak;
+			check(CompressedFileSize > 0);
+			const int32 CompressionBlockSize = FileEntry->CompressionBlockSize;
+			check((UncompressedFileSize + CompressionBlockSize - 1) / CompressionBlockSize == FileEntry->CompressionBlocks.Num());
+			Blocks.AddDefaulted(FileEntry->CompressionBlocks.Num());
+		}
+		UE_LOG(LogPakFile, Verbose, TEXT("FPakPlatformFile::OpenAsyncRead[%016llX, %016llX) %s"), OffsetInPak, OffsetInPak + CompressedFileSize, Filename);
+		check(PakFileSize > 0 && OffsetInPak + CompressedFileSize <= PakFileSize && OffsetInPak >= 0);
+
+		ReadCallbackFunction = [this](bool bWasCancelled, IAsyncReadRequest* Request)
+		{
+			RawReadCallback(bWasCancelled, Request);
+		};
+
 	}
+	~FPakAsyncReadFileHandle()
+	{
+		check(!LiveRequests.Num()); // must delete all requests before you delete the handle
+		check(!NumLiveRawRequests); // must delete all requests before you delete the handle
+	}
+
 	virtual IAsyncReadRequest* SizeRequest(FAsyncFileCallBack* CompleteCallback = nullptr) override
 	{
-		return new FPakSizeRequest(CompleteCallback, FileSize);
+		return new FPakSizeRequest(CompleteCallback, UncompressedFileSize);
 	}
-	virtual IAsyncReadRequest* ReadRequest(int64 Offset, int64 BytesToRead, EAsyncIOPriority Priority = AIOP_Normal, FAsyncFileCallBack* CompleteCallback = nullptr) override
+	virtual IAsyncReadRequest* ReadRequest(int64 Offset, int64 BytesToRead, EAsyncIOPriority Priority = AIOP_Normal, FAsyncFileCallBack* CompleteCallback = nullptr, uint8* UserSuppliedMemory = nullptr) override
 	{
 		if (BytesToRead == MAX_int64)
 		{
-			BytesToRead = FileSize - Offset;
+			BytesToRead = UncompressedFileSize - Offset;
+		}
+		check(Offset + BytesToRead <= UncompressedFileSize && Offset >= 0);
+		if (FileEntry->CompressionMethod == COMPRESS_None)
+		{
+			check(Offset + BytesToRead + OffsetInPak <= PakFileSize);
+			check(!Blocks.Num());
+			return new FPakReadRequest(PakFile, PakFileSize, CompleteCallback, OffsetInPak + Offset, BytesToRead, Priority, UserSuppliedMemory);
+		}
+		bool bAnyUnfinished = false;
+		FPakProcessedReadRequest* Result;
+		{
+			FScopeLock ScopedLock(&CriticalSection);
+			check(Blocks.Num());
+			int32 FirstBlock = Offset / FileEntry->CompressionBlockSize;
+			int32 LastBlock = (Offset + BytesToRead - 1) / FileEntry->CompressionBlockSize;
+
+			check(FirstBlock >= 0 && FirstBlock < Blocks.Num() && LastBlock >= 0 && LastBlock < Blocks.Num() && FirstBlock <= LastBlock);
+
+			Result = new FPakProcessedReadRequest(this, CompleteCallback, Offset, BytesToRead, Priority, UserSuppliedMemory);
+			for (int32 BlockIndex = FirstBlock; BlockIndex <= LastBlock; BlockIndex++)
+			{
+
+				FCachedAsyncBlock& Block = Blocks[BlockIndex];
+				Block.RefCount++;
+				if (!Block.bInFlight)
+				{
+					StartBlock(BlockIndex, Priority);
+					bAnyUnfinished = true;
+				}
+				if (!Block.Processed)
+				{
+					bAnyUnfinished = true;
+				}
+			}
+			if (Result)
+			{
+				check(!LiveRequests.Contains(Result))
+				LiveRequests.Add(Result);
+			}
+			if (!bAnyUnfinished)
+			{
+				Result->RequestIsComplete();
+			}
+		}
+		return Result;
+	}
+
+	void StartBlock(int32 BlockIndex, EAsyncIOPriority Priority)
+	{
+		FCachedAsyncBlock& Block = Blocks[BlockIndex];
+		Block.bInFlight = true;
+		check(!Block.RawRequest && !Block.Processed && !Block.Raw && !Block.CPUWorkGraphEvent.GetReference() && !Block.ProcessedSize && !Block.RawSize);
+		Block.RawSize = FileEntry->CompressionBlocks[BlockIndex].CompressedEnd - FileEntry->CompressionBlocks[BlockIndex].CompressedStart;
+		NumLiveRawRequests++;
+		Block.RawRequest = new FPakReadRequest(PakFile, PakFileSize, &ReadCallbackFunction, FileEntry->CompressionBlocks[BlockIndex].CompressedStart, Block.RawSize, Priority, nullptr, true, BlockIndex);
+	}
+	void RawReadCallback(bool bWasCancelled, IAsyncReadRequest* InRequest)
+	{
+		FPakReadRequest* Request = static_cast<FPakReadRequest*>(InRequest);
+		// Causes a deadlock, hopefully not needed as we are only referencing the block.
+		// Potential problem is with cancel
+		// FScopeLock ScopedLock(&CriticalSection);
+		int32 BlockIndex = Request->GetBlockIndex();
+		check(BlockIndex >= 0 && BlockIndex < Blocks.Num());
+		FCachedAsyncBlock& Block = Blocks[BlockIndex];
+		check((Block.RawRequest == Request || (!Block.RawRequest && Block.RawSize)) // we still might be in the constructor so the assignment hasn't happened yet
+			&& !Block.Processed && !Block.Raw);
+		if (bWasCancelled)
+		{
+			Block.RawSize = 0;
+		}
+		else
+		{
+			Block.Raw = Request->GetReadResults();
+			check(Block.Raw);
+			Block.ProcessedSize = FileEntry->CompressionBlockSize;
+			if (BlockIndex == Blocks.Num() - 1)
+			{
+				Block.ProcessedSize = FileEntry->UncompressedSize % FileEntry->CompressionBlockSize;
+				if (!Block.ProcessedSize)
+				{
+					Block.ProcessedSize = FileEntry->CompressionBlockSize; // last block was a full block
+				}
+			}
+			check(Block.ProcessedSize);
+			Block.CPUWorkGraphEvent = TGraphTask<FAsyncIOCPUWorkTask>::CreateTask().ConstructAndDispatchWhenReady(*this, BlockIndex);
+		}
+	}
+	void DoProcessing(int32 BlockIndex)
+	{
+		check(BlockIndex >= 0 && BlockIndex < Blocks.Num());
+		FCachedAsyncBlock& Block = Blocks[BlockIndex];
+		check(Block.Raw && Block.RawSize && !Block.Processed);
+
+		check(Block.ProcessedSize);
+		INC_MEMORY_STAT_BY(STAT_AsyncFileMemory, Block.ProcessedSize);
+		uint8* Output = (uint8*)FMemory::Malloc(Block.ProcessedSize);
+		FCompression::UncompressMemory((ECompressionFlags)FileEntry->CompressionMethod, Output, Block.ProcessedSize, Block.Raw, Block.RawSize, false, FPlatformMisc::GetPlatformCompression()->GetCompressionBitWindow());
+		FMemory::Free(Block.Raw);
+		Block.Raw = nullptr;
+		DEC_MEMORY_STAT_BY(STAT_AsyncFileMemory, Block.RawSize);
+		Block.RawSize = 0;
+
+		{
+			FScopeLock ScopedLock(&CriticalSection);
+			if (Block.RawRequest)
+			{
+				Block.RawRequest->WaitCompletion();
+				delete Block.RawRequest;
+				Block.RawRequest = nullptr;
+				NumLiveRawRequests--;
+			}
+			if (Block.RefCount > 0)
+			{
+				Block.Processed = Output;
+				for (FPakProcessedReadRequest* Req : LiveRequests)
+				{
+					if (Req->CheckCompletion(*FileEntry, BlockIndex, Blocks))
+					{
+						Req->RequestIsComplete();
+					}
+				}
+			}
+			else
+			{
+				// must have been canceled, clean up
+				FMemory::Free(Output);
+				Output = nullptr;
+				check(Block.ProcessedSize);
+				DEC_MEMORY_STAT_BY(STAT_AsyncFileMemory, Block.ProcessedSize);
+				Block.ProcessedSize = 0;
+				Block.CPUWorkGraphEvent = nullptr;
+				Block.bInFlight = false;
+			}
+		}
+	}
+	void ClearBlock(FCachedAsyncBlock& Block)
+	{
+		check(!Block.RawRequest);
+		Block.RawRequest = nullptr;
+		check(!Block.CPUWorkGraphEvent.GetReference() || Block.CPUWorkGraphEvent->IsComplete());
+		Block.CPUWorkGraphEvent = nullptr;
+		if (Block.Raw)
+		{
+			// this was a cancel, clean it up now
+			FMemory::Free(Block.Raw);
+			Block.Raw = nullptr;
+			check(Block.RawSize);
+			DEC_MEMORY_STAT_BY(STAT_AsyncFileMemory, Block.RawSize);
+		}
+		Block.RawSize = 0;
+		if (Block.Processed)
+		{
+			FMemory::Free(Block.Processed);
+			Block.Processed = nullptr;
+			check(Block.ProcessedSize);
+			DEC_MEMORY_STAT_BY(STAT_AsyncFileMemory, Block.ProcessedSize);
+		}
+		Block.ProcessedSize = 0;
+		Block.bInFlight = false;
+	}
+
+	void RemoveRequest(FPakProcessedReadRequest* Req, int64 Offset, int64 BytesToRead)
+	{
+		FScopeLock ScopedLock(&CriticalSection);
+		check(LiveRequests.Contains(Req));
+		LiveRequests.Remove(Req);
+		int32 FirstBlock = Offset / FileEntry->CompressionBlockSize;
+		int32 LastBlock = (Offset + BytesToRead - 1) / FileEntry->CompressionBlockSize;
+		check(FirstBlock >= 0 && FirstBlock < Blocks.Num() && LastBlock >= 0 && LastBlock < Blocks.Num() && FirstBlock <= LastBlock);
+
+		for (int32 BlockIndex = FirstBlock; BlockIndex <= LastBlock; BlockIndex++)
+		{
+			FCachedAsyncBlock& Block = Blocks[BlockIndex];
+			check(Block.RefCount > 0);
+			if (!--Block.RefCount)
+			{
+				if (Block.RawRequest)
+				{
+					Block.RawRequest->Cancel();
+					Block.RawRequest->WaitCompletion();
+					delete Block.RawRequest;
+					Block.RawRequest = nullptr;
+					NumLiveRawRequests--;
+				}
+				if (!Block.CPUWorkGraphEvent.GetReference() || Block.CPUWorkGraphEvent->IsComplete())
+				{
+					ClearBlock(Block);
+				}
+			}
+		}
+	}
+	void GatherResults(uint8* Memory, int64 Offset, int64 BytesToRead)
+	{
+		// no lock here, I don't think it is needed because we have a ref count.
+		int32 FirstBlock = Offset / FileEntry->CompressionBlockSize;
+		int32 LastBlock = (Offset + BytesToRead - 1) / FileEntry->CompressionBlockSize;
+		check(FirstBlock >= 0 && FirstBlock < Blocks.Num() && LastBlock >= 0 && LastBlock < Blocks.Num() && FirstBlock <= LastBlock);
+
+		for (int32 BlockIndex = FirstBlock; BlockIndex <= LastBlock; BlockIndex++)
+		{
+			FCachedAsyncBlock& Block = Blocks[BlockIndex];
+			check(Block.RefCount > 0 && Block.Processed && Block.ProcessedSize);
+			int64 BlockStart = int64(BlockIndex) * int64(FileEntry->CompressionBlockSize);
+			int64 BlockEnd = BlockStart + Block.ProcessedSize;
+
+			int64 SrcOffset = 0;
+			int64 DestOffset = BlockStart - Offset;
+			if (DestOffset < 0)
+			{
+				SrcOffset -= DestOffset;
+				DestOffset = 0;
+			}
+			int64 CopySize = Block.ProcessedSize;
+			if (DestOffset + CopySize > BytesToRead)
+			{
+				CopySize = BytesToRead - DestOffset;
+			}
+			if (SrcOffset + CopySize > Block.ProcessedSize)
+			{
+				CopySize = Block.ProcessedSize - SrcOffset;
+			}
+			check(CopySize > 0 && DestOffset >= 0 && DestOffset + CopySize <= BytesToRead);
+			check(SrcOffset >= 0 && SrcOffset + CopySize <= Block.ProcessedSize);
+			FMemory::Memcpy(Memory + DestOffset, Block.Processed + SrcOffset, CopySize);
+
+			check(Block.RefCount > 0);
 		}
-		check(Offset + BytesToRead <= FileSize && Offset >= 0 && Offset + BytesToRead + OffsetInPak <= PakFileSize);
-		return new FPakReadRequest(PakFile, PakFileSize, CompleteCallback, OffsetInPak + Offset, BytesToRead, Priority);
 	}
 };
 
+void FPakProcessedReadRequest::CancelRawRequests()
+{
+	if (CompleteRace.Increment() == 1)
+	{
+		Owner->RemoveRequest(this, Offset, BytesToRead);
+		bHasCancelled = true;
+	}
+}
+
+void FPakProcessedReadRequest::GatherResults()
+{
+	if (CompleteRace.Increment() == 1)
+	{
+		if (!bUserSuppliedMemory)
+		{
+			check(!Memory);
+			Memory = (uint8*)FMemory::Malloc(BytesToRead);
+			INC_MEMORY_STAT_BY(STAT_AsyncFileMemory, BytesToRead);
+		}
+		check(Memory);
+		Owner->GatherResults(Memory, Offset, BytesToRead);
+	}
+}
+
+void FPakProcessedReadRequest::DoneWithRawRequests()
+{
+	if (!bHasCancelled)
+	{
+		Owner->RemoveRequest(this, Offset, BytesToRead);
+	}
+}
+
+bool FPakProcessedReadRequest::CheckCompletion(const FPakEntry& FileEntry, int32 BlockIndex, TArray<FCachedAsyncBlock>& Blocks)
+{
+	if (!bRequestOutstanding || bHasCompleted)
+	{
+		return false;
+	}
+	{
+		int64 BlockStart = int64(BlockIndex) * int64(FileEntry.CompressionBlockSize);
+		int64 BlockEnd = int64(BlockIndex + 1) * int64(FileEntry.CompressionBlockSize);
+		if (Offset >= BlockEnd || Offset + BytesToRead <= BlockStart)
+		{
+			return false;
+		}
+	}
+	int32 FirstBlock = Offset / FileEntry.CompressionBlockSize;
+	int32 LastBlock = (Offset + BytesToRead - 1) / FileEntry.CompressionBlockSize;
+	check(FirstBlock >= 0 && FirstBlock < Blocks.Num() && LastBlock >= 0 && LastBlock < Blocks.Num() && FirstBlock <= LastBlock);
+
+	for (int32 MyBlockIndex = FirstBlock; MyBlockIndex <= LastBlock; MyBlockIndex++)
+	{
+		FCachedAsyncBlock& Block = Blocks[MyBlockIndex];
+		if (!Block.Processed)
+		{
+			return false;
+		}
+	}
+	bHasCompleted = true;
+	return true;
+}
+
+void FAsyncIOCPUWorkTask::DoTask(ENamedThreads::Type CurrentThread, const FGraphEventRef& MyCompletionGraphEvent)
+{
+	Owner.DoProcessing(BlockIndex);
+}
+
 #endif
 
 
@@ -2200,15 +2850,13 @@ public:
 IAsyncReadFileHandle* FPakPlatformFile::OpenAsyncRead(const TCHAR* Filename)
 {
 #if USE_PAK_PRECACHE
-	if (GPakCache_Enable > 0)
+	if (FPlatformProcess::SupportsMultithreading() && GPakCache_Enable > 0)
 	{
 		FPakFile* PakFile = NULL;
 		const FPakEntry* FileEntry = FindFileInPakFiles(Filename, &PakFile);
 		if (FileEntry && PakFile && PakFile->GetFilenameName() != NAME_None)
 		{
-			int64 OffsetToFile = FileEntry->Offset + FileEntry->GetSerializedSize(PakFile->GetInfo().Version);
-			UE_LOG(LogPakFile, Verbose, TEXT("FPakPlatformFile::OpenAsyncRead[%016llX, %016llX) %s"), OffsetToFile, OffsetToFile + FileEntry->UncompressedSize, Filename);
-			return new FPakAsyncReadFileHandle(PakFile->GetFilenameName(), PakFile->TotalSize(), OffsetToFile, FileEntry->UncompressedSize);
+			return new FPakAsyncReadFileHandle(FileEntry, PakFile, Filename);
 		}
 	}
 #endif
@@ -2434,29 +3082,29 @@ bool FPakEntry::VerifyPakEntriesMatch(const FPakEntry& FileEntryA, const FPakEnt
 	return bResult;
 }
 
-bool FPakPlatformFile::IsNonPakFilenameAllowed(const FString& InFilename, bool bAllowDirectories)
+bool FPakPlatformFile::IsNonPakFilenameAllowed(const FString& InFilename)
 {
-	FFilenameSecurityDelegate& FilenameSecurityDelegate = GetFilenameSecurityDelegate();
+	bool bAllowed = true;
 
-	if (!FilenameSecurityDelegate.IsBound())
-	{
-		return true;
-	}
+#if EXCLUDE_NONPAK_UE_EXTENSIONS
+	FName Ext = FName(*FPaths::GetExtension(InFilename));
+	bAllowed = !ExcludedNonPakExtensions.Contains(Ext);
+#endif
 
-	if (GetLowerLevel()->FileExists(*InFilename))
-	{
-		return FilenameSecurityDelegate.Execute(*InFilename);
-	}
-	else if (bAllowDirectories && GetLowerLevel()->DirectoryExists(*InFilename))
-	{
-		return FilenameSecurityDelegate.Execute(*InFilename);
-	}
-	else
+	FFilenameSecurityDelegate& FilenameSecurityDelegate = GetFilenameSecurityDelegate();
+	if (bAllowed)
 	{
-		return false;
+		if (FilenameSecurityDelegate.IsBound())
+		{
+			bAllowed = FilenameSecurityDelegate.Execute(*InFilename);;
+		}
 	}
+
+	return bAllowed;
 }
 
+
+#if IS_PROGRAM
 FPakFile::FPakFile(const TCHAR* Filename, bool bIsSigned)
 	: PakFilename(Filename)
 	, PakFilenameName(Filename)
@@ -2471,6 +3119,7 @@ FPakFile::FPakFile(const TCHAR* Filename, bool bIsSigned)
 		Initialize(Reader);
 	}
 }
+#endif
 
 FPakFile::FPakFile(IPlatformFile* LowerLevel, const TCHAR* Filename, bool bIsSigned)
 	: PakFilename(Filename)
@@ -2487,12 +3136,14 @@ FPakFile::FPakFile(IPlatformFile* LowerLevel, const TCHAR* Filename, bool bIsSig
 	}
 }
 
+#if WITH_EDITOR
 FPakFile::FPakFile(FArchive* Archive)
 	: bSigned(false)
 	, bIsValid(false)
 {
 	Initialize(Archive);
 }
+#endif
 
 FPakFile::~FPakFile()
 {
@@ -2782,7 +3433,7 @@ void FPakPlatformFile::HandlePakListCommand(const TCHAR* Cmd, FOutputDevice& Ar)
 	GetMountedPaks(Paks);
 	for (auto Pak : Paks)
 	{
-		Ar.Logf(TEXT("%s"), *Pak.PakFile->GetFilename());
+		Ar.Logf(TEXT("%s Mounted to %s"), *Pak.PakFile->GetFilename(), *Pak.PakFile->GetMountPoint());
 	}	
 }
 #endif // !UE_BUILD_SHIPPING
@@ -2798,12 +3449,17 @@ FPakPlatformFile::~FPakPlatformFile()
 	FCoreDelegates::OnMountPak.Unbind();
 	FCoreDelegates::OnUnmountPak.Unbind();
 
+#if USE_PAK_PRECACHE
+	FPakPrecacher::Shutdown();
+#endif
+#if !USE_NEW_ASYNC_IO
 	// We need to flush async IO... if it hasn't been shut down already.
 	if (FIOSystem::HasShutdown() == false)
 	{
 		FIOSystem& IOSystem = FIOSystem::Get();
 		IOSystem.BlockTillAllRequestsFinishedAndFlushHandles();
 	}
+#endif
 
 	{
 		FScopeLock ScopedLock(&PakListCritical);
@@ -2923,6 +3579,14 @@ bool FPakPlatformFile::Initialize(IPlatformFile* Inner, const TCHAR* CmdLine)
 	check(Inner != NULL);
 	LowerLevel = Inner;
 
+#if EXCLUDE_NONPAK_UE_EXTENSIONS
+	// Extensions for file types that should only ever be in a pak file. Used to stop unnecessary access to the lower level platform file
+	ExcludedNonPakExtensions.Add(TEXT("uasset"));
+	ExcludedNonPakExtensions.Add(TEXT("umap"));
+	ExcludedNonPakExtensions.Add(TEXT("ubulk"));
+	ExcludedNonPakExtensions.Add(TEXT("uexp"));
+#endif
+
 #if !USING_SIGNED_CONTENT
 	bSigned = FParse::Param(CmdLine, TEXT("Signedpak")) || FParse::Param(CmdLine, TEXT("Signed"));
 	if (!bSigned)
@@ -3005,7 +3669,14 @@ bool FPakPlatformFile::Initialize(IPlatformFile* Inner, const TCHAR* CmdLine)
 	FCoreDelegates::OnMountPak.BindRaw(this, &FPakPlatformFile::HandleMountPakDelegate);
 	FCoreDelegates::OnUnmountPak.BindRaw(this, &FPakPlatformFile::HandleUnmountPakDelegate);
 #if USE_PAK_PRECACHE
-	FPakPrecacher::Init(LowerLevel);
+	if (!WITH_EDITOR && FPlatformProcess::SupportsMultithreading())
+	{
+		FPakPrecacher::Init(LowerLevel);
+	}
+	else
+	{
+		GPakCache_Enable = 0;
+	}
 #endif
 
 	return !!LowerLevel;
@@ -3054,6 +3725,12 @@ bool FPakPlatformFile::Mount(const TCHAR* InPakFilename, uint32 PakOrder, const
 
 bool FPakPlatformFile::Unmount(const TCHAR* InPakFilename)
 {
+#if USE_PAK_PRECACHE
+	if (GPakCache_Enable)
+	{
+		FPakPrecacher::Get().Unmount(InPakFilename);
+	}
+#endif
 	{
 		FScopeLock ScopedLock(&PakListCritical); 
 
@@ -3067,7 +3744,6 @@ bool FPakPlatformFile::Unmount(const TCHAR* InPakFilename)
 			}
 		}
 	}
-
 	return false;
 }
 
diff --git a/Engine/Source/Runtime/PakFile/Public/IPlatformFilePak.h b/Engine/Source/Runtime/PakFile/Public/IPlatformFilePak.h
--- a/Engine/Source/Runtime/PakFile/Public/IPlatformFilePak.h
+++ b/Engine/Source/Runtime/PakFile/Public/IPlatformFilePak.h
@@ -288,13 +288,15 @@ class PAKFILE_API FPakFile : FNoncopyable
 
 public:
 
+#if IS_PROGRAM
 	/**
-	 * Opens a pak file given its filename.
-	 *
-	 * @param Filename Pak filename.
-	 * @param bIsSigned true if the pak is signed
-	 */
-	FPakFile(const TCHAR* Filename, bool bIsSigned);	
+	* Opens a pak file given its filename.
+	*
+	* @param Filename Pak filename.
+	* @param bIsSigned true if the pak is signed
+	*/
+	FPakFile(const TCHAR* Filename, bool bIsSigned);
+#endif
 
 	/**
 	 * Creates a pak file using the supplied file handle.
@@ -310,7 +312,9 @@ public:
 	 *
 	 * @param Archive	Pointer to the archive which contains the pak file data.
 	 */
+#if WITH_EDITOR
 	FPakFile(FArchive* Archive);
+#endif
 
 	~FPakFile();
 
@@ -835,6 +839,8 @@ class PAKFILE_API FPakPlatformFile : public IPlatformFile
 	bool bSigned;
 	/** Synchronization object for accessing the list of currently mounted pak files. */
 	FCriticalSection PakListCritical;
+	/** Cache of extensions that we automatically reject if not found in pak file */
+	TSet<FName> ExcludedNonPakExtensions;
 
 	/**
 	 * Gets mounted pak files
@@ -923,7 +929,7 @@ class PAKFILE_API FPakPlatformFile : public IPlatformFile
 	 * @param InFilename			Filename to check
 	 * @param bAllowDirectories		Consider directories as valid filepaths?
 	 */
-	bool IsNonPakFilenameAllowed(const FString& InFilename, bool bAllowDirectories = false);
+	bool IsNonPakFilenameAllowed(const FString& InFilename);
 
 public:
 
@@ -1132,7 +1138,13 @@ public:
 		FDateTime Result = FDateTime::MinValue();
 		if (IsNonPakFilenameAllowed(Filename))
 		{
+#if USE_NEW_ASYNC_IO
+			double StartTime(UE_LOG_ACTIVE(LogPakFile, Error) ? FPlatformTime::Seconds() : 0.0);
+#endif
 			Result = LowerLevel->GetTimeStamp(Filename);
+#if USE_NEW_ASYNC_IO
+			UE_LOG(LogPakFile, Error, TEXT("GetTimeStamp on disk (!!) for %s took %6.2fms."), Filename, float(FPlatformTime::Seconds() - StartTime) * 1000.0f);
+#endif
 		}
 		return Result;
 	}
@@ -1300,8 +1312,7 @@ public:
 
 		// Fall back to lower level.
 		FFileStatData FileStatData;
-		const bool bAllowDirectories = true;
-		if (IsNonPakFilenameAllowed(FilenameOrDirectory, bAllowDirectories))
+		if (IsNonPakFilenameAllowed(FilenameOrDirectory))
 		{
 			FileStatData = LowerLevel->GetStatData(FilenameOrDirectory);
 		}
diff --git a/Engine/Source/Runtime/RHI/Public/RHIDefinitions.h b/Engine/Source/Runtime/RHI/Public/RHIDefinitions.h
--- a/Engine/Source/Runtime/RHI/Public/RHIDefinitions.h
+++ b/Engine/Source/Runtime/RHI/Public/RHIDefinitions.h
@@ -509,10 +509,8 @@ enum ETextureCreateFlags
 	// Texture that may be updated every frame
 	TexCreate_Dynamic				= 1<<8,
 	// Allow silent texture creation failure
-	// @warning:	When you update this, you must update FTextureAllocations::FindTextureType() in Core/Private/UObject/TextureAllocations.cpp
 	TexCreate_AllowFailure			= 1<<9,
 	// Disable automatic defragmentation if the initial texture memory allocation fails.
-	// @warning:	When you update this, you must update FTextureAllocations::FindTextureType() in Core/Private/UObject/TextureAllocations.cpp
 	TexCreate_DisableAutoDefrag		= 1<<10,
 	// Create the texture with automatic -1..1 biasing
 	TexCreate_BiasNormalMap			= 1<<11,
diff --git a/Engine/Source/Runtime/SandboxFile/Public/IPlatformFileSandboxWrapper.h b/Engine/Source/Runtime/SandboxFile/Public/IPlatformFileSandboxWrapper.h
--- a/Engine/Source/Runtime/SandboxFile/Public/IPlatformFileSandboxWrapper.h
+++ b/Engine/Source/Runtime/SandboxFile/Public/IPlatformFileSandboxWrapper.h
@@ -613,7 +613,12 @@ public:
 #if USE_NEW_ASYNC_IO
 	virtual IAsyncReadFileHandle* OpenAsyncRead(const TCHAR* Filename) override
 	{
-		return LowerLevel->OpenAsyncRead(*ConvertToSandboxPath(Filename));
+		FString UserFilename(*ConvertToSandboxPath(Filename));
+		if (!OkForInnerAccess(Filename) || LowerLevel->FileExists(*UserFilename))
+		{
+			return LowerLevel->OpenAsyncRead(*UserFilename);
+		}
+		return LowerLevel->OpenAsyncRead(Filename);
 	}
 #endif // USE_NEW_ASYNC_IO
 };
