Copying //UE4/Dev-Core to //UE4/Dev-Main (Source: //UE4/Dev-Core @ 3208226)

diff --git a/Engine/Source/Runtime/Engine/Private/World.cpp b/Engine/Source/Runtime/Engine/Private/World.cpp
--- a/Engine/Source/Runtime/Engine/Private/World.cpp
+++ b/Engine/Source/Runtime/Engine/Private/World.cpp
@@ -473,20 +473,20 @@ bool UWorld::Rename(const TCHAR* InName, UObject* NewOuter, ERenameFlags Flags)
 
 		if (PersistentLevel->MapBuildData->IsLegacyBuildData())
 		{
-			TArray<UTexture2D*> LightMapsAndShadowMaps;
-			GetLightMapsAndShadowMaps(PersistentLevel, LightMapsAndShadowMaps);
+	TArray<UTexture2D*> LightMapsAndShadowMaps;
+	GetLightMapsAndShadowMaps(PersistentLevel, LightMapsAndShadowMaps);
 
-			for (auto* Tex : LightMapsAndShadowMaps)
-			{
+	for (auto* Tex : LightMapsAndShadowMaps)
+	{
 				if (Tex)
-				{
+		{
 					if (!Tex->Rename(*Tex->GetName(), NewOuter, Flags))
-					{
-						return false;
-					}
-				}
+			{
+				return false;
 			}
 		}
+	}
+		}
 		else
 		{
 			FString NewPackageName = GetOutermost()->GetName() + TEXT("_BuiltData");
@@ -494,15 +494,15 @@ bool UWorld::Rename(const TCHAR* InName, UObject* NewOuter, ERenameFlags Flags)
 			UPackage* BuildDataPackage = PersistentLevel->MapBuildData->GetOutermost();
 
 			if (!BuildDataPackage->Rename(*NewPackageName, NewOuter, Flags))
-			{
-				return false;
-			}
+	{
+		return false;
+	}
 		}
 
 		if (!PersistentLevel->MapBuildData->Rename(*NewMapBuildDataName.ToString(), NewOuter, Flags))
-		{
-			return false;
-		}
+	{
+		return false;
+	}
 	}
 
 	// Rename the level script blueprint now, unless we are in PostLoad. ULevel::PostLoad should handle renaming this at load time.
@@ -616,19 +616,19 @@ void UWorld::PostDuplicate(bool bDuplicateForPIE)
 				NewTextureOuter = NewBuildData;
 			}
 
-			TArray<UTexture2D*> LightMapsAndShadowMaps;
-			GetLightMapsAndShadowMaps(PersistentLevel, LightMapsAndShadowMaps);
+		TArray<UTexture2D*> LightMapsAndShadowMaps;
+		GetLightMapsAndShadowMaps(PersistentLevel, LightMapsAndShadowMaps);
 
-			// Duplicate the textures, if any
-			for (auto* Tex : LightMapsAndShadowMaps)
-			{
+		// Duplicate the textures, if any
+		for (auto* Tex : LightMapsAndShadowMaps)
+		{
 				if (Tex && Tex->GetOutermost() != NewTextureOuter)
-				{
+			{
 					UObject* NewTex = StaticDuplicateObject(Tex, NewTextureOuter, Tex->GetFName());
-					ReplacementMap.Add(Tex, NewTex);
-				}
+				ReplacementMap.Add(Tex, NewTex);
 			}
 		}
+		}
 
 		// Duplicate the level script blueprint generated classes as well
 		const bool bDontCreate = true;
@@ -1120,13 +1120,13 @@ void UWorld::InitWorld(const InitializationValues IVS)
 		if (WorldSettings)
 		{
 			if (IVS.bCreateNavigation && WorldSettings->bEnableNavigationSystem)
-			{
-				UNavigationSystem::CreateNavigationSystem(this);
-			}
+		{
+			UNavigationSystem::CreateNavigationSystem(this);
+		}
 			if (IVS.bCreateAISystem && WorldSettings->bEnableAISystem)
-			{
-				CreateAISystem();
-			}
+	{
+		CreateAISystem();
+	}
 		}
 	}
 	
@@ -1615,8 +1615,8 @@ void UWorld::UpdateCullDistanceVolumes(AActor* ActorToUpdate, UPrimitiveComponen
 				check((ActorToUpdate == nullptr) || (ActorToUpdate == ComponentToUpdate->GetOwner()));
 				if (ACullDistanceVolume::CanBeAffectedByVolumes(ComponentToUpdate))
 				{
-					CompToNewMaxDrawMap.Add(ComponentToUpdate, ComponentToUpdate->LDMaxDrawDistance);
-				}
+				CompToNewMaxDrawMap.Add(ComponentToUpdate, ComponentToUpdate->LDMaxDrawDistance);
+			}
 			}
 			else
 			{
@@ -1878,13 +1878,16 @@ DEFINE_STAT(STAT_UpdateLevelStreamingTime);
  *
  * @return true if time limit has been exceeded, false otherwise
  */
-static bool IsTimeLimitExceeded( const TCHAR* CurrentTask, double StartTime, ULevel* Level )
+static bool IsTimeLimitExceeded( const TCHAR* CurrentTask, double StartTime, ULevel* Level, double TimeLimit = 0.0 )
 {
 	bool bIsTimeLimitExceed = false;
 	// We don't spread work across several frames in the Editor to avoid potential side effects.
 	if( Level->OwningWorld->IsGameWorld() == true )
 	{
-		double TimeLimit = GLevelStreamingActorsUpdateTimeLimit;
+		if (TimeLimit == 0.0)
+		{
+			TimeLimit = GLevelStreamingActorsUpdateTimeLimit;
+		}
 		double CurrentTime	= FPlatformTime::Seconds();
 		// Delta time in ms.
 		double DeltaTime	= (CurrentTime - StartTime) * 1000;
@@ -1949,7 +1952,7 @@ private:
 void UWorld::AddToWorld( ULevel* Level, const FTransform& LevelTransform )
 {
 	SCOPE_CYCLE_COUNTER(STAT_AddToWorldTime);
-	
+
 	check(Level);
 	check(!Level->IsPendingKill());
 	check(!Level->IsUnreachable());
@@ -1965,10 +1968,11 @@ void UWorld::AddToWorld( ULevel* Level, const FTransform& LevelTransform )
 	// Don't consider the time limit if the match hasn't started as we need to ensure that the levels are fully loaded
 	const bool bConsiderTimeLimit = bMatchStarted;
 
-	bool bExecuteNextStep = (CurrentLevelPendingVisibility == Level) || CurrentLevelPendingVisibility == NULL;
+	bool bExecuteNextStep = CurrentLevelPendingVisibility == Level || CurrentLevelPendingVisibility == NULL;
 	bool bPerformedLastStep	= false;
 	
-	if( bExecuteNextStep && CurrentLevelPendingVisibility == NULL )
+	// Don't make this level visible if it's currently being made invisible
+	if( bExecuteNextStep && CurrentLevelPendingVisibility == NULL && CurrentLevelPendingInvisibility != Level )
 	{
 		Level->OwningWorld = this;
 		
@@ -2221,21 +2225,57 @@ void UWorld::AddToWorld( ULevel* Level, const FTransform& LevelTransform )
 #endif // PERF_TRACK_DETAILED_ASYNC_STATS
 }
 
-void UWorld::RemoveFromWorld( ULevel* Level )
+void UWorld::RemoveFromWorld( ULevel* Level, bool bAllowIncrementalRemoval )
 {
 	SCOPE_CYCLE_COUNTER(STAT_RemoveFromWorldTime);
+	
 	FScopeCycleCounterUObject Context(Level);
 	check(Level);
 	check(!Level->IsPendingKill());
 	check(!Level->IsUnreachable());
 
-	if (CurrentLevelPendingVisibility == NULL && Level->bIsVisible)
+	if ( CurrentLevelPendingVisibility == NULL && Level->bIsVisible )
 	{
 		// Keep track of timing.
 		double StartTime = FPlatformTime::Seconds();	
 
+		bool bFinishRemovingLevel = true;
+		if ( bAllowIncrementalRemoval && GLevelStreamingUnregisterComponentsTimeLimit > 0.0 )
+		{
+			bFinishRemovingLevel = false;
+			if (CurrentLevelPendingInvisibility == NULL)
+			{
+				// Mark level as being the one in process of being made invisible. 
+				// This will prevent this level from being unloaded or made visible in the meantime
+				CurrentLevelPendingInvisibility = Level;
+				Level->bIsBeingRemoved = true;
+			}
+
+			if ( CurrentLevelPendingInvisibility == Level )
+			{
+				// Incrementally unregister actor components. 
+				// This avoids spikes on the renderthread and gamethread when we subsequently call ClearLevelComponents() further down
+				check(IsGameWorld());
+				int32 NumComponentsToUnregister = GLevelStreamingComponentsUnregistrationGranularity;
+				do
+				{
+					if (Level->IncrementalUnregisterComponents(NumComponentsToUnregister))
+					{
+						// We're done, so the level can be removed
+						CurrentLevelPendingInvisibility = NULL;
+						bFinishRemovingLevel = true;
+						break;
+					}
+				} while (!IsTimeLimitExceeded(TEXT("unregistering components"), StartTime, Level, GLevelStreamingUnregisterComponentsTimeLimit));
+			}
+		}
+        else
+        {
 		Level->bIsBeingRemoved = true;
+        }
 
+		if ( bFinishRemovingLevel )
+		{
 		for (int32 ActorIdx = 0; ActorIdx < Level->Actors.Num(); ActorIdx++)
 		{
 			AActor* Actor = Level->Actors[ActorIdx];
@@ -2253,7 +2293,7 @@ void UWorld::RemoveFromWorld( ULevel* Level )
 		// Remove any pawns from the pawn list that are about to be streamed out
 		for( FConstPawnIterator Iterator = GetPawnIterator(); Iterator; ++Iterator )
 		{
-			APawn* Pawn = *Iterator;
+			APawn* Pawn = Iterator->Get();
 			if (Pawn->IsInLevel(Level))
 			{
 				RemovePawn(Pawn);
@@ -2322,6 +2362,7 @@ void UWorld::RemoveFromWorld( ULevel* Level )
 		ULevelStreaming::BroadcastLevelVisibleStatus(this, Level->GetOutermost()->GetFName(), false);
 
 		Level->bIsBeingRemoved = false;
+		} // if ( bFinishRemovingLevel )
 
 #if PERF_TRACK_DETAILED_ASYNC_STATS
 		UE_LOG(LogStreaming, Display, TEXT("UWorld::RemoveFromWorld for %s took %5.2f ms"), *Level->GetOutermost()->GetName(), (FPlatformTime::Seconds() - StartTime) * 1000.0);
@@ -2649,19 +2690,24 @@ void UWorld::UpdateLevelStreamingInner(ULevelStreaming* StreamingLevel)
 	// Work performed to make a level visible is spread across several frames and we can't unload/ hide a level that is currently pending
 	// to be made visible, so we fulfill those requests first.
 	bool bHasVisibilityRequestPending	= StreamingLevel->GetLoadedLevel() && StreamingLevel->GetLoadedLevel() == CurrentLevelPendingVisibility;
+	bool bHasInvisibilityRequestPending	= StreamingLevel->GetLoadedLevel() && StreamingLevel->GetLoadedLevel() == CurrentLevelPendingInvisibility;
 		
 	// Figure out whether level should be loaded, visible and block on load if it should be loaded but currently isn't.
 	bool bShouldBeLoaded = bHasVisibilityRequestPending || (!GUseBackgroundLevelStreaming && !bShouldForceUnloadStreamingLevels && !StreamingLevel->bIsRequestingUnloadAndRemoval);
 	bool bShouldBeVisible	= bHasVisibilityRequestPending || bShouldForceVisibleStreamingLevels;
 	bool bShouldBlockOnLoad	= StreamingLevel->bShouldBlockOnLoad || StreamingLevel->ShouldBeAlwaysLoaded();
 
-	// Don't update if the code requested this level object to be unloaded and removed.
-	if(!bShouldForceUnloadStreamingLevels && !StreamingLevel->bIsRequestingUnloadAndRemoval)
+	// Don't update if the code requested this level object to be unloaded and removed or it has an invisibility request pending
+	if ( bHasInvisibilityRequestPending )
+	{
+		// If we have an invisibily request pending, we can't make the level visible again until the request is completed
+		bShouldBeVisible = false;
+	}
+	else if (!bShouldForceUnloadStreamingLevels && !StreamingLevel->bIsRequestingUnloadAndRemoval)
 	{
 		bShouldBeLoaded		= bShouldBeLoaded  || !IsGameWorld() || StreamingLevel->ShouldBeLoaded();
 		bShouldBeVisible	= bShouldBeVisible || (bShouldBeLoaded && StreamingLevel->ShouldBeVisible());
 	}
-
 	// We want to give the garbage collector a chance to remove levels before we stream in more. We can't do this in the
 	// case of a blocking load as it means those requests should be fulfilled right away. By waiting on GC before kicking
 	// off new levels we potentially delay streaming in maps, but AllowLevelLoadRequests already looks and checks whether
@@ -2722,10 +2768,10 @@ void UWorld::UpdateLevelStreamingInner(ULevelStreaming* StreamingLevel)
 		{
 			// Discard previous LOD level
 			StreamingLevel->DiscardPendingUnloadLevel(this);
-			// Hide loaded level
 			if (Level->bIsVisible)
 			{
-				RemoveFromWorld(Level);
+				// Hide loaded level, incrementally if necessary
+				RemoveFromWorld(Level, IsGameWorld() );
 			}
 		}
 
@@ -2877,7 +2923,7 @@ void UWorld::ConditionallyBuildStreamingData()
 
 bool UWorld::IsVisibilityRequestPending() const
 {
-	return (CurrentLevelPendingVisibility != NULL);
+	return (CurrentLevelPendingVisibility != nullptr && CurrentLevelPendingInvisibility != nullptr);
 }
 
 bool UWorld::AreAlwaysLoadedLevelsLoaded() const
@@ -3495,12 +3541,14 @@ UGameViewportClient* UWorld::GetGameViewport() const
 
 FConstControllerIterator UWorld::GetControllerIterator() const
 {
-	return ControllerList.CreateConstIterator();
+	auto Result = ControllerList.CreateConstIterator();
+	return (const FConstControllerIterator&)Result;
 }
 
 FConstPlayerControllerIterator UWorld::GetPlayerControllerIterator() const
 {
-	return PlayerControllerList.CreateConstIterator();
+	auto Result = PlayerControllerList.CreateConstIterator();
+	return (const FConstPlayerControllerIterator&)Result;
 }
 
 APlayerController* UWorld::GetFirstPlayerController() const
@@ -3508,7 +3556,7 @@ APlayerController* UWorld::GetFirstPlayerController() const
 	APlayerController* PlayerController = NULL;
 	if( GetPlayerControllerIterator() )
 	{
-		PlayerController = *GetPlayerControllerIterator();
+		PlayerController = GetPlayerControllerIterator()->Get();
 	}
 	return PlayerController;
 }
@@ -3517,7 +3565,7 @@ ULocalPlayer* UWorld::GetFirstLocalPlayerFromController() const
 {
 	for( FConstPlayerControllerIterator Iterator = GetPlayerControllerIterator(); Iterator; ++Iterator )
 	{
-		APlayerController* PlayerController = *Iterator;
+		APlayerController* PlayerController = Iterator->Get();
 		if( PlayerController )
 		{
 			ULocalPlayer* LocalPlayer = Cast<ULocalPlayer>(PlayerController->Player);
@@ -3555,7 +3603,8 @@ void UWorld::RemoveController( AController* Controller )
 
 FConstPawnIterator UWorld::GetPawnIterator() const
 {
-	return PawnList.CreateConstIterator();
+	auto Result = PawnList.CreateConstIterator();
+	return (const FConstPawnIterator&)Result;
 }
 
 int32 UWorld::GetNumPawns() const
@@ -3593,7 +3642,8 @@ void UWorld::RegisterAutoActivateCamera(ACameraActor* CameraActor, int32 PlayerI
 
 FConstCameraActorIterator UWorld::GetAutoActivateCameraIterator() const
 {
-	return AutoCameraActorList.CreateConstIterator();
+	auto Result = AutoCameraActorList.CreateConstIterator();
+	return (const FConstCameraActorIterator&)Result;
 }
 
 
@@ -3735,7 +3785,8 @@ void UWorld::RemovePhysicsVolume(APhysicsVolume* Volume)
 
 FConstPhysicsVolumeIterator UWorld::GetNonDefaultPhysicsVolumeIterator() const
 {
-	return NonDefaultPhysicsVolumeList.CreateConstIterator();
+	auto Result = NonDefaultPhysicsVolumeList.CreateConstIterator();
+	return (const FConstPhysicsVolumeIterator&)Result;
 }
 
 int32 UWorld::GetNonDefaultPhysicsVolumeCount() const
@@ -3979,7 +4030,7 @@ bool UWorld::DestroySwappedPC(UNetConnection* Connection)
 {
 	for( FConstPlayerControllerIterator Iterator = GetPlayerControllerIterator(); Iterator; ++Iterator )
 	{
-		APlayerController* PlayerController = *Iterator;
+		APlayerController* PlayerController = Iterator->Get();
 		if (PlayerController->Player == NULL && PlayerController->PendingSwapConnection == Connection)
 		{
 			DestroyActor(PlayerController);
@@ -5069,7 +5120,7 @@ UWorld* FSeamlessTravelHandler::Tick()
 			{
 				for( FConstControllerIterator Iterator = CurrentWorld->GetControllerIterator(); Iterator; ++Iterator )
 				{
-					AController* Player = *Iterator;
+					AController* Player = Iterator->Get();
 					if (Player->PlayerState || Cast<APlayerController>(Player) != nullptr)
 					{
 						KeepAnnotation.Set(Player);
diff --git a/Engine/Source/Runtime/Engine/Private/WorldComposition.cpp b/Engine/Source/Runtime/Engine/Private/WorldComposition.cpp
--- a/Engine/Source/Runtime/Engine/Private/WorldComposition.cpp
+++ b/Engine/Source/Runtime/Engine/Private/WorldComposition.cpp
@@ -424,9 +424,9 @@ void UWorldComposition::OnTileInfoUpdated(const FName& InPackageName, const FWor
 	UPackage* LevelPackage = Cast<UPackage>(StaticFindObjectFast(UPackage::StaticClass(), NULL, Tile->PackageName));
 	if (LevelPackage)
 	{
-		if (LevelPackage->WorldTileInfo == NULL)
+		if (LevelPackage->WorldTileInfo == nullptr)
 		{
-			LevelPackage->WorldTileInfo = new FWorldTileInfo(Tile->Info);
+			LevelPackage->WorldTileInfo = MakeUnique<FWorldTileInfo>(Tile->Info);
 			PackageDirty = true;
 		}
 		else
@@ -887,7 +887,7 @@ void UWorldComposition::OnLevelPostLoad(ULevel* InLevel)
 		const bool bIsDefault = (Info == FWorldTileInfo());
 		if (!bIsDefault)
 		{
-			LevelPackage->WorldTileInfo = new FWorldTileInfo(Info);
+			LevelPackage->WorldTileInfo = MakeUnique<FWorldTileInfo>(Info);
 		}
 	}
 }
diff --git a/Engine/Source/Runtime/Engine/Public/DistanceFieldAtlas.h b/Engine/Source/Runtime/Engine/Public/DistanceFieldAtlas.h
--- a/Engine/Source/Runtime/Engine/Public/DistanceFieldAtlas.h
+++ b/Engine/Source/Runtime/Engine/Public/DistanceFieldAtlas.h
@@ -8,6 +8,7 @@
 
 #include "TextureLayout3d.h"
 #include "GCObject.h"
+#include "UniquePtr.h"
 
 // DDC key for distance field data, must be changed when modifying the generation code or data format
 #define DISTANCEFIELD_DERIVEDDATA_VER TEXT("7768798764B445A9543C94442EA899D")
@@ -226,7 +227,7 @@ private:
 	void Build(FAsyncDistanceFieldTask* Task, class FQueuedThreadPool& ThreadPool);
 
 	/** Thread that will build any tasks in TaskQueue and exit when there are no more. */
-	class TScopedPointer<class FBuildDistanceFieldThreadRunnable> ThreadRunnable;
+	class TUniquePtr<class FBuildDistanceFieldThreadRunnable> ThreadRunnable;
 
 	/** Game-thread managed list of tasks in the async system. */
 	TArray<FAsyncDistanceFieldTask*> ReferencedTasks;
diff --git a/Engine/Source/Runtime/Engine/Public/GlobalShader.h b/Engine/Source/Runtime/Engine/Public/GlobalShader.h
--- a/Engine/Source/Runtime/Engine/Public/GlobalShader.h
+++ b/Engine/Source/Runtime/Engine/Public/GlobalShader.h
@@ -8,6 +8,7 @@
 
 #include "Shader.h"
 #include "SceneManagement.h"
+#include "UniquePtr.h"
 
 class FShaderCommonCompileJob;
 class FShaderCompileJob;
@@ -162,7 +163,7 @@ public:
 */
 struct FGlobalShaderBackupData
 {
-	TScopedPointer<TArray<uint8>> FeatureLevelShaderData[ERHIFeatureLevel::Num];
+	TUniquePtr<TArray<uint8>> FeatureLevelShaderData[ERHIFeatureLevel::Num];
 };
 
 /** Backs up all global shaders to memory through serialization, and removes all references to FShaders from the global shader map. */
diff --git a/Engine/Source/Runtime/Engine/Public/MaterialShared.h b/Engine/Source/Runtime/Engine/Public/MaterialShared.h
--- a/Engine/Source/Runtime/Engine/Public/MaterialShared.h
+++ b/Engine/Source/Runtime/Engine/Public/MaterialShared.h
@@ -13,6 +13,7 @@
 #include "StaticParameterSet.h"
 #include "Shader.h"
 #include "VertexFactory.h"
+#include "Optional.h"
 
 class FMaterialShaderMap;
 class FMaterialShaderType;
@@ -241,7 +242,7 @@ protected:
 class FUniformExpressionSet : public FRefCountedObject
 {
 public:
-	FUniformExpressionSet(): UniformBufferStruct(NULL) {}
+	FUniformExpressionSet() {}
 
 	ENGINE_API void Serialize(FArchive& Ar);
 	bool IsEmpty() const;
@@ -283,7 +284,7 @@ protected:
 	TArray<FGuid> ParameterCollections;
 
 	/** The structure of a uniform buffer containing values for these uniform expressions. */
-	TScopedPointer<FUniformBufferStruct> UniformBufferStruct;
+	TOptional<FUniformBufferStruct> UniformBufferStruct;
 
 	friend class FMaterial;
 	friend class FHLSLMaterialTranslator;
@@ -1199,9 +1200,9 @@ public:
 	static void UpdateEditorLoadedMaterialResources(EShaderPlatform InShaderPlatform);
 
 	/** Backs up any FShaders in editor loaded materials to memory through serialization and clears FShader references. */
-	static void BackupEditorLoadedMaterialShadersToMemory(TMap<FMaterialShaderMap*, TScopedPointer<TArray<uint8> > >& ShaderMapToSerializedShaderData);
+	static void BackupEditorLoadedMaterialShadersToMemory(TMap<FMaterialShaderMap*, TUniquePtr<TArray<uint8> > >& ShaderMapToSerializedShaderData);
 	/** Recreates FShaders in editor loaded materials from the passed in memory, handling shader key changes. */
-	static void RestoreEditorLoadedMaterialShadersFromMemory(const TMap<FMaterialShaderMap*, TScopedPointer<TArray<uint8> > >& ShaderMapToSerializedShaderData);
+	static void RestoreEditorLoadedMaterialShadersFromMemory(const TMap<FMaterialShaderMap*, TUniquePtr<TArray<uint8> > >& ShaderMapToSerializedShaderData);
 
 protected:
 	
@@ -1734,9 +1735,9 @@ class FMaterialUpdateContext
 	/** Materials updated within this context. */
 	TSet<UMaterialInterface*> UpdatedMaterialInterfaces;
 	/** Active global component reregister context, if any. */
-	TScopedPointer<class FGlobalComponentReregisterContext> ComponentReregisterContext;
+	TUniquePtr<class FGlobalComponentReregisterContext> ComponentReregisterContext;
 	/** Active global component render state recreation context, if any. */
-	TScopedPointer<class FGlobalComponentRecreateRenderStateContext> ComponentRecreateRenderStateContext;
+	TUniquePtr<class FGlobalComponentRecreateRenderStateContext> ComponentRecreateRenderStateContext;
 	/** The shader platform that was being processed - can control if we need to update components */
 	EShaderPlatform ShaderPlatform;
 	/** True if the SyncWithRenderingThread option was specified. */
diff --git a/Engine/Source/Runtime/Engine/Public/Model.h b/Engine/Source/Runtime/Engine/Public/Model.h
--- a/Engine/Source/Runtime/Engine/Public/Model.h
+++ b/Engine/Source/Runtime/Engine/Public/Model.h
@@ -11,6 +11,7 @@
 #include "SceneTypes.h"
 #include "RawIndexBuffer.h"
 #include "Classes/Engine/MapBuildDataRegistry.h"
+#include "UniquePtr.h"
 
 struct FLightmassPrimitiveSettings;
 struct FStaticLightingVertex;
@@ -372,7 +373,7 @@ class UModel : public UObject
 	TArray<FLightmassPrimitiveSettings>	LightmassSettings;
 
 	/** An index buffer for each material used by the model, containing all the nodes with that material applied. */
-	TMap<UMaterialInterface*,TScopedPointer<FRawIndexBuffer16or32> > MaterialIndexBuffers;
+	TMap<UMaterialInterface*,TUniquePtr<FRawIndexBuffer16or32> > MaterialIndexBuffers;
 
 	/** A vertex buffer containing the vertices for all nodes in the UModel. */
 	FModelVertexBuffer VertexBuffer;
diff --git a/Engine/Source/Runtime/Engine/Public/ShaderCompiler.h b/Engine/Source/Runtime/Engine/Public/ShaderCompiler.h
--- a/Engine/Source/Runtime/Engine/Public/ShaderCompiler.h
+++ b/Engine/Source/Runtime/Engine/Public/ShaderCompiler.h
@@ -7,6 +7,7 @@
 #pragma once
 
 #include "ShaderCore.h"
+#include "UniquePtr.h"
 
 DECLARE_LOG_CATEGORY_EXTERN(LogShaderCompilers, Log, All);
 
@@ -340,7 +341,7 @@ private:
 	TMap<int32, FShaderMapFinalizeResults> PendingFinalizeShaderMaps;
 
 	/** The threads spawned for shader compiling. */
-	TScopedPointer<FShaderCompileThreadRunnableBase> Thread;
+	TUniquePtr<FShaderCompileThreadRunnableBase> Thread;
 
 	//////////////////////////////////////////////////////
 	// Configuration properties - these are set only on initialization and can be read from either thread
diff --git a/Engine/Source/Runtime/Engine/Public/StaticMeshResources.h b/Engine/Source/Runtime/Engine/Public/StaticMeshResources.h
--- a/Engine/Source/Runtime/Engine/Public/StaticMeshResources.h
+++ b/Engine/Source/Runtime/Engine/Public/StaticMeshResources.h
@@ -15,6 +15,7 @@
 #include "SceneManagement.h"
 #include "PhysicsEngine/BodySetupEnums.h"
 #include "Materials/MaterialInterface.h"
+#include "UniquePtr.h"
 
 /**
  * The LOD settings to use for a group of static meshes.
@@ -846,7 +847,7 @@ public:
 	void SyncUVChannelData(const TArray<FStaticMaterial>& ObjectData);
 
 	/** The next cached derived data in the list. */
-	TScopedPointer<class FStaticMeshRenderData> NextCachedRenderData;
+	TUniquePtr<class FStaticMeshRenderData> NextCachedRenderData;
 
 	/**
 	 * Cache derived renderable data for the static mesh with the provided
diff --git a/Engine/Source/Runtime/Engine/Public/TextureResource.h b/Engine/Source/Runtime/Engine/Public/TextureResource.h
--- a/Engine/Source/Runtime/Engine/Public/TextureResource.h
+++ b/Engine/Source/Runtime/Engine/Public/TextureResource.h
@@ -8,6 +8,7 @@
 
 #include "Engine/TextureDefines.h"
 #include "UnrealClient.h"
+#include "UniquePtr.h"
 
 class UTexture2D;
 
@@ -311,7 +312,7 @@ private:
 	int32 CurrentFirstMip;
 
 	/** Pending async create texture task, if any.															*/
-	TScopedPointer<FAsyncCreateTextureTask> AsyncCreateTextureTask;
+	TUniquePtr<FAsyncCreateTextureTask> AsyncCreateTextureTask;
 	
 	/** Local copy/ cache of mip data between creation and first call to InitRHI.							*/
 	void*				MipData[MAX_TEXTURE_MIP_COUNT];
diff --git a/Engine/Source/Runtime/JsonUtilities/Private/JsonObjectConverter.cpp b/Engine/Source/Runtime/JsonUtilities/Private/JsonObjectConverter.cpp
--- a/Engine/Source/Runtime/JsonUtilities/Private/JsonObjectConverter.cpp
+++ b/Engine/Source/Runtime/JsonUtilities/Private/JsonObjectConverter.cpp
@@ -28,7 +28,14 @@ TSharedPtr<FJsonValue> ConvertScalarUPropertyToJsonValue(UProperty* Property, co
 		// fall through to default cases
 	}
 
-	if (UNumericProperty *NumericProperty = Cast<UNumericProperty>(Property))
+	if (UEnumProperty* EnumProperty = Cast<UEnumProperty>(Property))
+	{
+		// export enums as strings
+		UEnum* EnumDef = EnumProperty->GetEnum();
+		FString StringValue = EnumDef->GetEnumName(EnumProperty->GetUnderlyingProperty()->GetSignedIntPropertyValue(Value));
+		return MakeShareable(new FJsonValueString(StringValue));
+	}
+	else if (UNumericProperty *NumericProperty = Cast<UNumericProperty>(Property))
 	{
 		// see if it's an enum
 		UEnum* EnumDef = NumericProperty->GetIntPropertyEnum();
@@ -50,7 +57,7 @@ TSharedPtr<FJsonValue> ConvertScalarUPropertyToJsonValue(UProperty* Property, co
 		}
 
 		// fall through to default
-	}		
+	}
 	else if (UBoolProperty *BoolProperty = Cast<UBoolProperty>(Property))
 	{
 		// Export bools as bools
@@ -248,21 +255,43 @@ bool GetTextFromObject(const TSharedRef<FJsonObject>& Obj, FText& TextOut)
 /** Convert JSON to property, assuming either the property is not an array or the value is an individual array element */
 bool ConvertScalarJsonValueToUProperty(TSharedPtr<FJsonValue> JsonValue, UProperty* Property, void* OutValue, int64 CheckFlags, int64 SkipFlags)
 {
-	if (UNumericProperty *NumericProperty = Cast<UNumericProperty>(Property))
+	if (UEnumProperty* EnumProperty = Cast<UEnumProperty>(Property))
+	{
+		if (JsonValue->Type == EJson::String)
+		{
+			// see if we were passed a string for the enum
+			const UEnum* Enum = EnumProperty->GetEnum();
+			check(Enum);
+			FString StrValue = JsonValue->AsString();
+			int64 IntValue = Enum->GetValueByName(FName(*StrValue));
+			if (IntValue == INDEX_NONE)
+			{
+				UE_LOG(LogJson, Error, TEXT("JsonValueToUProperty - Unable import enum %s from string value %s for property %s"), *Enum->CppType, *StrValue, *Property->GetNameCPP());
+				return false;
+			}
+			EnumProperty->GetUnderlyingProperty()->SetIntPropertyValue(OutValue, IntValue);
+		}
+		else
+		{
+			// AsNumber will log an error for completely inappropriate types (then give us a default)
+			EnumProperty->GetUnderlyingProperty()->SetIntPropertyValue(OutValue, (int64)JsonValue->AsNumber());
+		}
+	}
+	else if (UNumericProperty *NumericProperty = Cast<UNumericProperty>(Property))
 	{
 		if (NumericProperty->IsEnum() && JsonValue->Type == EJson::String)
 		{
 			// see if we were passed a string for the enum
-			const UEnum* EnumProperty = NumericProperty->GetIntPropertyEnum();
-			check(EnumProperty); // should be assured by IsEnum()
+			const UEnum* Enum = NumericProperty->GetIntPropertyEnum();
+			check(Enum); // should be assured by IsEnum()
 			FString StrValue = JsonValue->AsString();
-			int32 IntValue = EnumProperty->GetValueByName(FName(*StrValue));
+			int64 IntValue = Enum->GetValueByName(FName(*StrValue));
 			if (IntValue == INDEX_NONE)
 			{
-				UE_LOG(LogJson, Error, TEXT("JsonValueToUProperty - Unable import enum %s from string value %s for property %s"), *EnumProperty->CppType, *StrValue, *Property->GetNameCPP());
+				UE_LOG(LogJson, Error, TEXT("JsonValueToUProperty - Unable import enum %s from string value %s for property %s"), *Enum->CppType, *StrValue, *Property->GetNameCPP());
 				return false;
 			}
-			NumericProperty->SetIntPropertyValue(OutValue, (int64)IntValue);
+			NumericProperty->SetIntPropertyValue(OutValue, IntValue);
 		}
 		else if(NumericProperty->IsFloatingPoint())
 		{
diff --git a/Engine/Source/Runtime/Landscape/Private/LandscapeGrass.cpp b/Engine/Source/Runtime/Landscape/Private/LandscapeGrass.cpp
--- a/Engine/Source/Runtime/Landscape/Private/LandscapeGrass.cpp
+++ b/Engine/Source/Runtime/Landscape/Private/LandscapeGrass.cpp
@@ -2148,7 +2148,7 @@ void ALandscapeProxy::UpdateGrass(const TArray<FVector>& Cameras, bool bForceSyn
 
 											for (auto& LOD : HierarchicalInstancedStaticMeshComponent->LODData)
 											{
-												LOD.OverrideMapBuildData = TScopedPointer<FMeshMapBuildData>(new FMeshMapBuildData());
+												LOD.OverrideMapBuildData = MakeUnique<FMeshMapBuildData>();
 												LOD.OverrideMapBuildData->LightMap = GrassLightMap;
 												LOD.OverrideMapBuildData->ShadowMap = GrassShadowMap;
 											}
diff --git a/Engine/Source/Runtime/Launch/Private/HTML5/LaunchHTML5.cpp b/Engine/Source/Runtime/Launch/Private/HTML5/LaunchHTML5.cpp
--- a/Engine/Source/Runtime/Launch/Private/HTML5/LaunchHTML5.cpp
+++ b/Engine/Source/Runtime/Launch/Private/HTML5/LaunchHTML5.cpp
@@ -94,7 +94,7 @@ public:
 		return false;
 	}
 };
-static TAutoPtr<FHTML5Exec> GHTML5Exec;
+static TUniquePtr<FHTML5Exec> GHTML5Exec;
 
 int main(int argc, char* argv[])
 {
@@ -110,7 +110,7 @@ int main(int argc, char* argv[])
 
 	// TODO: configure this via the command line?
 	emscripten_trace_configure("http://127.0.0.1:5000/", "UE4Game");
-	GHTML5Exec = new FHTML5Exec();
+	GHTML5Exec = MakeUnique<FHTML5Exec>();
 
 	emscripten_trace_enter_context("main");
 
diff --git a/Engine/Source/Runtime/Launch/Private/LaunchEngineLoop.cpp b/Engine/Source/Runtime/Launch/Private/LaunchEngineLoop.cpp
--- a/Engine/Source/Runtime/Launch/Private/LaunchEngineLoop.cpp
+++ b/Engine/Source/Runtime/Launch/Private/LaunchEngineLoop.cpp
@@ -21,6 +21,7 @@
 #include "MallocProfiler.h"
 
 #include "NetworkVersion.h"
+#include "UniquePtr.h"
 
 #if WITH_COREUOBJECT
 	#include "Internationalization/PackageLocalizationManager.h"
@@ -238,9 +239,9 @@ public:
 };
 
 
-static TScopedPointer<FOutputDeviceConsole>	GScopedLogConsole;
-static TScopedPointer<FOutputDeviceStdOutput> GScopedStdOut;
-static TScopedPointer<FOutputDeviceTestExit> GScopedTestExit;
+static TUniquePtr<FOutputDeviceConsole>	GScopedLogConsole;
+static TUniquePtr<FOutputDeviceStdOutput> GScopedStdOut;
+static TUniquePtr<FOutputDeviceTestExit> GScopedTestExit;
 
 
 #if WITH_ENGINE
@@ -268,10 +269,10 @@ static void RHIExitAndStopRHIThread()
 void InitializeStdOutDevice()
 {
 	// Check if something is trying to initialize std out device twice.
-	check(!GScopedStdOut.IsValid());
+	check(!GScopedStdOut);
 
-	GScopedStdOut = new FOutputDeviceStdOutput();
-	GLog->AddOutputDevice(GScopedStdOut.GetOwnedPointer());
+	GScopedStdOut = MakeUnique<FOutputDeviceStdOutput>();
+	GLog->AddOutputDevice(GScopedStdOut.Get());
 }
 
 
@@ -842,7 +843,7 @@ int32 FEngineLoop::PreInit( const TCHAR* CmdLine )
 	}
 
 	// Initialize log console here to avoid statics initialization issues when launched from the command line.
-	GScopedLogConsole = FPlatformOutputDevices::GetLogConsole();
+	GScopedLogConsole = TUniquePtr<FOutputDeviceConsole>(FPlatformOutputDevices::GetLogConsole());
 
 	// Always enable the backlog so we get all messages, we will disable and clear it in the game
 	// as soon as we determine whether GIsEditor == false
@@ -863,8 +864,8 @@ int32 FEngineLoop::PreInit( const TCHAR* CmdLine )
 			TArray<FString> ExitPhrasesList;
 			if (ExitPhrases.ParseIntoArray(ExitPhrasesList, TEXT("+"), true) > 0)
 			{
-				GScopedTestExit = new FOutputDeviceTestExit(ExitPhrasesList);
-				GLog->AddOutputDevice(GScopedTestExit.GetOwnedPointer());
+				GScopedTestExit = MakeUnique<FOutputDeviceTestExit>(ExitPhrasesList);
+				GLog->AddOutputDevice(GScopedTestExit.Get());
 			}
 		}
 	}
@@ -1303,7 +1304,7 @@ int32 FEngineLoop::PreInit( const TCHAR* CmdLine )
 	}
 
 	// Get a pointer to the log output device
-	GLogConsole = GScopedLogConsole.GetOwnedPointer();
+	GLogConsole = GScopedLogConsole.Get();
 
 	LoadPreInitModules();
 
@@ -1456,7 +1457,7 @@ int32 FEngineLoop::PreInit( const TCHAR* CmdLine )
 
 	// If std out device hasn't been initialized yet (there was no -stdout param in the command line) and
 	// we meet all the criteria, initialize it now.
-	if (!GScopedStdOut.IsValid() && !bHasEditorToken && !bIsRegularClient && !IsRunningDedicatedServer())
+	if (!GScopedStdOut && !bHasEditorToken && !bIsRegularClient && !IsRunningDedicatedServer())
 	{
 		InitializeStdOutDevice();
 	}
@@ -2589,7 +2590,7 @@ void FEngineLoop::ProcessLocalPlayerSlateOperations() const
 				{
 					for (FConstPlayerControllerIterator Iterator = CurWorld->GetPlayerControllerIterator(); Iterator; ++Iterator)
 					{
-						APlayerController* PlayerController = *Iterator;
+						APlayerController* PlayerController = Iterator->Get();
 						if (PlayerController)
 						{
 							ULocalPlayer* LocalPlayer = Cast< ULocalPlayer >(PlayerController->Player);
diff --git a/Engine/Source/Runtime/MovieScene/Private/MovieSceneCommonHelpers.cpp b/Engine/Source/Runtime/MovieScene/Private/MovieSceneCommonHelpers.cpp
--- a/Engine/Source/Runtime/MovieScene/Private/MovieSceneCommonHelpers.cpp
+++ b/Engine/Source/Runtime/MovieScene/Private/MovieSceneCommonHelpers.cpp
@@ -346,6 +346,23 @@ FTrackInstancePropertyBindings::FPropertyAddress FTrackInstancePropertyBindings:
 	}
 }
 
+void FTrackInstancePropertyBindings::CallFunctionForEnum( UObject& InRuntimeObject, int64 PropertyValue )
+{
+	FPropertyAndFunction PropAndFunction = FindOrAdd(InRuntimeObject);
+	if (PropAndFunction.Function)
+	{
+		// ProcessEvent should really be taking const void*
+		InRuntimeObject.ProcessEvent(PropAndFunction.Function, (void*)&PropertyValue);
+	}
+	else if (PropAndFunction.PropertyAddress.Address)
+	{
+		UEnumProperty* EnumProperty = CastChecked<UEnumProperty>(PropAndFunction.PropertyAddress.Property);
+		UNumericProperty* UnderlyingProperty = EnumProperty->GetUnderlyingProperty();
+		void* ValueAddr = EnumProperty->ContainerPtrToValuePtr<void>(PropAndFunction.PropertyAddress.Address);
+		UnderlyingProperty->SetIntPropertyValue(ValueAddr, PropertyValue);
+	}
+}
+
 void FTrackInstancePropertyBindings::CacheBinding(const UObject& Object)
 {
 	FPropertyAndFunction PropAndFunction;
@@ -375,6 +392,22 @@ UProperty* FTrackInstancePropertyBindings::GetProperty(const UObject& Object) co
 	return Prop;
 }
 
+int64 FTrackInstancePropertyBindings::GetCurrentValueForEnum(const UObject& Object)
+{
+	FPropertyAndFunction PropAndFunction = FindOrAdd(Object);
+
+	if(PropAndFunction.PropertyAddress.Address)
+	{
+		UEnumProperty* EnumProperty = CastChecked<UEnumProperty>(PropAndFunction.PropertyAddress.Property);
+		UNumericProperty* UnderlyingProperty = EnumProperty->GetUnderlyingProperty();
+		void* ValueAddr = EnumProperty->ContainerPtrToValuePtr<void>(PropAndFunction.PropertyAddress.Address);
+		int64 Result = UnderlyingProperty->GetSignedIntPropertyValue(ValueAddr);
+		return Result;
+	}
+
+	return 0;
+}
+
 template<> void FTrackInstancePropertyBindings::CallFunction<bool>(UObject& InRuntimeObject, TCallTraits<bool>::ParamType PropertyValue)
 {
 	FPropertyAndFunction PropAndFunction = FindOrAdd(InRuntimeObject);
diff --git a/Engine/Source/Runtime/MovieScene/Public/MovieSceneCommonHelpers.h b/Engine/Source/Runtime/MovieScene/Public/MovieSceneCommonHelpers.h
--- a/Engine/Source/Runtime/MovieScene/Public/MovieSceneCommonHelpers.h
+++ b/Engine/Source/Runtime/MovieScene/Public/MovieSceneCommonHelpers.h
@@ -124,6 +124,14 @@ public:
 		}
 	}
 
+	/**
+	 * Calls the setter function for a specific runtime object or if the setter function does not exist, the property is set directly
+	 *
+	 * @param InRuntimeObject The runtime object whose function to call
+	 * @param PropertyValue The new value to assign to the property
+	 */
+	void CallFunctionForEnum( UObject& InRuntimeObject, int64 PropertyValue );
+
 	/**
 	 * Rebuilds the property and function mappings for a single runtime object, and adds them to the cache
 	 *
@@ -162,6 +170,14 @@ public:
 		return ValueType();
 	}
 
+	/**
+	 * Gets the current value of a property on an object
+	 *
+	 * @param Object	The object to get the property from
+	 * @return ValueType	The current value
+	 */
+	int64 GetCurrentValueForEnum(const UObject& Object);
+
 	/**
 	 * Sets the current value of a property on an object
 	 *
diff --git a/Engine/Source/Runtime/MovieSceneTracks/Private/Evaluation/MovieScenePropertyTemplates.cpp b/Engine/Source/Runtime/MovieSceneTracks/Private/Evaluation/MovieScenePropertyTemplates.cpp
--- a/Engine/Source/Runtime/MovieSceneTracks/Private/Evaluation/MovieScenePropertyTemplates.cpp
+++ b/Engine/Source/Runtime/MovieSceneTracks/Private/Evaluation/MovieScenePropertyTemplates.cpp
@@ -5,6 +5,7 @@
 #include "MovieSceneBoolSection.h"
 #include "MovieSceneFloatSection.h"
 #include "MovieSceneByteSection.h"
+#include "MovieSceneEnumSection.h"
 #include "MovieSceneIntegerSection.h"
 #include "MovieSceneVectorSection.h"
 #include "MovieSceneStringSection.h"
@@ -71,6 +72,23 @@ void FMovieSceneBytePropertySectionTemplate::Evaluate(const FMovieSceneEvaluatio
 	ExecutionTokens.Add(TCachedPropertyTrackExecutionToken<uint8>());
 }
 
+FMovieSceneEnumPropertySectionTemplate::FMovieSceneEnumPropertySectionTemplate(const UMovieSceneEnumSection& Section, const UMovieScenePropertyTrack& Track)
+	: PropertyData(Track.GetPropertyName(), Track.GetPropertyPath())
+	, EnumCurve(Section.GetCurve())
+{}
+
+void FMovieSceneEnumPropertySectionTemplate::Evaluate(const FMovieSceneEvaluationOperand& Operand, const FMovieSceneContext& Context, const FPersistentEvaluationData& PersistentData, FMovieSceneExecutionTokens& ExecutionTokens) const
+{
+	using namespace PropertyTemplate;
+
+	for (TCachedValue<int64>& ObjectAndValue : PersistentData.GetSectionData<TCachedSectionData<int64>>().ObjectsAndValues)
+	{
+		ObjectAndValue.Value = EnumCurve.Evaluate(Context.GetTime(), ObjectAndValue.Value);
+	}
+	
+	ExecutionTokens.Add(TCachedPropertyTrackExecutionToken<int64>());
+}
+
 FMovieSceneIntegerPropertySectionTemplate::FMovieSceneIntegerPropertySectionTemplate(const UMovieSceneIntegerSection& Section, const UMovieScenePropertyTrack& Track)
 	: PropertyData(Track.GetPropertyName(), Track.GetPropertyPath())
 	, IntegerCurve(Section.GetCurve())
diff --git a/Engine/Source/Runtime/MovieSceneTracks/Private/Evaluation/MovieScenePropertyTemplates.h b/Engine/Source/Runtime/MovieSceneTracks/Private/Evaluation/MovieScenePropertyTemplates.h
--- a/Engine/Source/Runtime/MovieSceneTracks/Private/Evaluation/MovieScenePropertyTemplates.h
+++ b/Engine/Source/Runtime/MovieSceneTracks/Private/Evaluation/MovieScenePropertyTemplates.h
@@ -17,6 +17,7 @@ class UMovieSceneIntegerSection;
 class UMovieSceneVectorSection;
 class UMovieSceneStringSection;
 class UMovieScenePropertyTrack;
+class UMovieSceneEnumSection;
 
 
 USTRUCT()
@@ -127,6 +128,42 @@ protected:
 	FIntegralCurve ByteCurve;
 };
 
+USTRUCT()
+struct FMovieSceneEnumPropertySectionTemplate : public FMovieSceneEvalTemplate
+{
+	GENERATED_BODY()
+	
+	FMovieSceneEnumPropertySectionTemplate(){}
+	FMovieSceneEnumPropertySectionTemplate(const UMovieSceneEnumSection& Section, const UMovieScenePropertyTrack& Track);
+
+protected:
+
+	virtual UScriptStruct& GetScriptStructImpl() const override
+	{
+		return *StaticStruct();
+	}
+	virtual void SetupOverrides() override
+	{
+		EnableOverrides(RequiresSetupFlag | RequiresInitializeFlag);
+	}
+	virtual void Setup(FPersistentEvaluationData& PersistentData, IMovieScenePlayer& Player) const override
+	{
+		PropertyData.SetupCachedTrack<int64>(PersistentData);
+	}
+	virtual void Initialize(const FMovieSceneEvaluationOperand& Operand, const FMovieSceneContext& Context, FPersistentEvaluationData& PersistentData, IMovieScenePlayer& Player) const override
+	{
+		PropertyData.SetupCachedFrame<int64>(Operand, PersistentData, Player);
+	}
+
+	virtual void Evaluate(const FMovieSceneEvaluationOperand& Operand, const FMovieSceneContext& Context, const FPersistentEvaluationData& PersistentData, FMovieSceneExecutionTokens& ExecutionTokens) const override;
+
+	UPROPERTY()
+	FMovieScenePropertySectionData PropertyData;
+
+	UPROPERTY()
+	FIntegralCurve EnumCurve;
+};
+
 USTRUCT()
 struct FMovieSceneIntegerPropertySectionTemplate : public FMovieSceneEvalTemplate
 {
diff --git a/Engine/Source/Runtime/MovieSceneTracks/Private/Sections/MovieSceneEnumSection.cpp b/Engine/Source/Runtime/MovieSceneTracks/Private/Sections/MovieSceneEnumSection.cpp
new file mode 100644
--- /dev/null
+++ b/Engine/Source/Runtime/MovieSceneTracks/Private/Sections/MovieSceneEnumSection.cpp
@@ -0,0 +1,99 @@
+// Copyright 1998-2016 Epic Games, Inc. All Rights Reserved.
+
+#include "MovieSceneTracksPrivatePCH.h"
+#include "MovieSceneEnumSection.h"
+#include "Evaluation/MovieScenePropertyTemplates.h"
+
+
+UMovieSceneEnumSection::UMovieSceneEnumSection( const FObjectInitializer& ObjectInitializer )
+	: Super( ObjectInitializer )
+{ }
+
+
+void UMovieSceneEnumSection::MoveSection( float DeltaPosition, TSet<FKeyHandle>& KeyHandles )
+{
+	Super::MoveSection( DeltaPosition, KeyHandles );
+
+	EnumCurve.ShiftCurve(DeltaPosition, KeyHandles);
+}
+
+
+void UMovieSceneEnumSection::DilateSection( float DilationFactor, float Origin, TSet<FKeyHandle>& KeyHandles )
+{
+	Super::DilateSection(DilationFactor, Origin, KeyHandles);
+	
+	EnumCurve.ScaleCurve(Origin, DilationFactor, KeyHandles);
+}
+
+
+void UMovieSceneEnumSection::GetKeyHandles(TSet<FKeyHandle>& OutKeyHandles, TRange<float> TimeRange) const
+{
+	if (!TimeRange.Overlaps(GetRange()))
+	{
+		return;
+	}
+
+	for (auto It(EnumCurve.GetKeyHandleIterator()); It; ++It)
+	{
+		float Time = EnumCurve.GetKeyTime(It.Key());
+		if (TimeRange.Contains(Time))
+		{
+			OutKeyHandles.Add(It.Key());
+		}
+	}
+}
+
+
+TOptional<float> UMovieSceneEnumSection::GetKeyTime( FKeyHandle KeyHandle ) const
+{
+	if ( EnumCurve.IsKeyHandleValid( KeyHandle ) )
+	{
+		return TOptional<float>( EnumCurve.GetKeyTime( KeyHandle ) );
+	}
+	return TOptional<float>();
+}
+
+
+void UMovieSceneEnumSection::SetKeyTime( FKeyHandle KeyHandle, float Time )
+{
+	if ( EnumCurve.IsKeyHandleValid( KeyHandle ) )
+	{
+		EnumCurve.SetKeyTime( KeyHandle, Time );
+	}
+}
+
+
+void UMovieSceneEnumSection::AddKey( float Time, const int64& Value, EMovieSceneKeyInterpolation KeyInterpolation )
+{
+	if (TryModify())
+	{
+		EnumCurve.UpdateOrAddKey(Time, Value);
+	}
+}
+
+
+bool UMovieSceneEnumSection::NewKeyIsNewData( float Time, const int64& Value ) const
+{
+	return EnumCurve.Evaluate( Time, Value ) != Value;
+}
+
+
+bool UMovieSceneEnumSection::HasKeys( const int64& Value ) const
+{
+	return EnumCurve.GetNumKeys() > 0;
+}
+
+
+void UMovieSceneEnumSection::SetDefault( const int64& Value )
+{
+	if (EnumCurve.GetDefaultValue() != Value && TryModify())
+	{
+		EnumCurve.SetDefaultValue(Value);
+	}
+}
+
+
+void UMovieSceneEnumSection::ClearDefaults()
+{
+	EnumCurve.ClearDefaultValue();
+}
diff --git a/Engine/Source/Runtime/MovieSceneTracks/Private/Tracks/MovieSceneEnumTrack.cpp b/Engine/Source/Runtime/MovieSceneTracks/Private/Tracks/MovieSceneEnumTrack.cpp
new file mode 100644
--- /dev/null
+++ b/Engine/Source/Runtime/MovieSceneTracks/Private/Tracks/MovieSceneEnumTrack.cpp
@@ -0,0 +1,33 @@
+// Copyright 1998-2016 Epic Games, Inc. All Rights Reserved.
+
+#include "MovieSceneTracksPrivatePCH.h"
+#include "MovieSceneEnumSection.h"
+#include "MovieSceneEnumTrack.h"
+#include "IMovieScenePlayer.h"
+#include "Evaluation/MovieScenePropertyTemplates.h"
+
+UMovieSceneEnumTrack::UMovieSceneEnumTrack( const FObjectInitializer& ObjectInitializer )
+	: Super( ObjectInitializer )
+{ }
+
+
+UMovieSceneSection* UMovieSceneEnumTrack::CreateNewSection()
+{
+	return NewObject<UMovieSceneSection>(this, UMovieSceneEnumSection::StaticClass(), NAME_None, RF_Transactional);
+}
+
+FMovieSceneEvalTemplatePtr UMovieSceneEnumTrack::CreateTemplateForSection(const UMovieSceneSection& InSection) const
+{
+	return FMovieSceneEnumPropertySectionTemplate(*CastChecked<UMovieSceneEnumSection>(&InSection), *this);
+}
+
+void UMovieSceneEnumTrack::SetEnum(UEnum* InEnum)
+{
+	Enum = InEnum;
+}
+
+
+class UEnum* UMovieSceneEnumTrack::GetEnum() const
+{
+	return Enum;
+}
diff --git a/Engine/Source/Runtime/MovieSceneTracks/Public/Sections/IKeyframeSection.h b/Engine/Source/Runtime/MovieSceneTracks/Public/Sections/IKeyframeSection.h
--- a/Engine/Source/Runtime/MovieSceneTracks/Public/Sections/IKeyframeSection.h
+++ b/Engine/Source/Runtime/MovieSceneTracks/Public/Sections/IKeyframeSection.h
@@ -22,4 +22,25 @@ public:
 	/** Removes all default data from this section. */
 	virtual void ClearDefaults() = 0;
 
+};
+
+class IKeyframeSectionEnum
+{
+public:
+
+	/** Returns true if the supplied key data is different than the section value at the supplied time, otherwise false. */
+	virtual bool NewKeyIsNewData( float Time, const int64& KeyData ) const = 0;
+
+	/** Returns true if the section has any keys which correspond to the supplied key data, otherwise false. */
+	virtual bool HasKeys( const int64& KeyData ) const = 0;
+
+	/** Adds a key to the section at the supplied time, with the supplied data. */
+	virtual void AddKey( float Time, const int64& KeyData, EMovieSceneKeyInterpolation KeyInterpolation ) = 0;
+
+	/** Sets the default value for the section which corresponds to the supplied key data. */
+	virtual void SetDefault( const int64& KeyData ) = 0;
+
+	/** Removes all default data from this section. */
+	virtual void ClearDefaults() = 0;
+
 };
\ No newline at end of file
diff --git a/Engine/Source/Runtime/MovieSceneTracks/Public/Sections/MovieSceneEnumSection.h b/Engine/Source/Runtime/MovieSceneTracks/Public/Sections/MovieSceneEnumSection.h
new file mode 100644
--- /dev/null
+++ b/Engine/Source/Runtime/MovieSceneTracks/Public/Sections/MovieSceneEnumSection.h
@@ -0,0 +1,53 @@
+// Copyright 1998-2016 Epic Games, Inc. All Rights Reserved.
+
+#pragma once
+
+#include "Curves/IntegralCurve.h"
+#include "MovieSceneSection.h"
+#include "IKeyframeSection.h"
+#include "MovieSceneEnumSection.generated.h"
+
+
+/**
+ * A single enum section.
+ */
+UCLASS(MinimalAPI)
+class UMovieSceneEnumSection 
+	: public UMovieSceneSection
+	, public IKeyframeSectionEnum
+{
+	GENERATED_UCLASS_BODY()
+public:
+
+	/** Gets all the keys of this enum section */
+	FIntegralCurve& GetCurve() { return EnumCurve; }
+	const FIntegralCurve& GetCurve() const { return EnumCurve; }
+
+public:
+
+	//~ IKeyframeSection interface
+
+	virtual void AddKey(float Time, const int64& Value, EMovieSceneKeyInterpolation KeyInterpolation) override;
+	virtual bool NewKeyIsNewData(float Time, const int64& Value) const override;
+	virtual bool HasKeys(const int64& Value) const override;
+	virtual void SetDefault(const int64& Value) override;
+	virtual void ClearDefaults() override;
+
+public:
+
+	//~ UMovieSceneSection interface 
+
+	virtual void MoveSection(float DeltaPosition, TSet<FKeyHandle>& KeyHandles) override;
+	virtual void DilateSection(float DilationFactor, float Origin, TSet<FKeyHandle>& KeyHandles) override;
+	virtual void GetKeyHandles(TSet<FKeyHandle>& OutKeyHandles, TRange<float> TimeRange) const override;
+	virtual TOptional<float> GetKeyTime(FKeyHandle KeyHandle) const override;
+	virtual void SetKeyTime(FKeyHandle KeyHandle, float Time) override;
+
+private:
+
+	/** Ordered curve data */
+	// @todo Sequencer This could be optimized by packing the enums separately
+	// but that may not be worth the effort
+	UPROPERTY()
+	FIntegralCurve EnumCurve;
+};
diff --git a/Engine/Source/Runtime/MovieSceneTracks/Public/Tracks/MovieSceneEnumTrack.h b/Engine/Source/Runtime/MovieSceneTracks/Public/Tracks/MovieSceneEnumTrack.h
new file mode 100644
--- /dev/null
+++ b/Engine/Source/Runtime/MovieSceneTracks/Public/Tracks/MovieSceneEnumTrack.h
@@ -0,0 +1,30 @@
+// Copyright 1998-2016 Epic Games, Inc. All Rights Reserved.
+
+#pragma once
+
+#include "MovieScene.h"
+#include "MovieSceneTrack.h"
+#include "MovieScenePropertyTrack.h"
+#include "MovieSceneEnumTrack.generated.h"
+
+/**
+ * Handles manipulation of byte properties in a movie scene
+ */
+UCLASS()
+class MOVIESCENETRACKS_API UMovieSceneEnumTrack : public UMovieScenePropertyTrack
+{
+	GENERATED_UCLASS_BODY()
+
+public:
+	/** UMovieSceneTrack interface */
+	virtual UMovieSceneSection* CreateNewSection() override;
+	virtual FMovieSceneEvalTemplatePtr CreateTemplateForSection(const UMovieSceneSection& InSection) const override;
+
+	void SetEnum(UEnum* Enum);
+
+	class UEnum* GetEnum() const;
+
+protected:
+	UPROPERTY()
+	UEnum* Enum;
+};
diff --git a/Engine/Source/Runtime/NetworkFile/Private/NetworkPlatformFile.cpp b/Engine/Source/Runtime/NetworkFile/Private/NetworkPlatformFile.cpp
--- a/Engine/Source/Runtime/NetworkFile/Private/NetworkPlatformFile.cpp
+++ b/Engine/Source/Runtime/NetworkFile/Private/NetworkPlatformFile.cpp
@@ -5,11 +5,13 @@
 #include "MultichannelTCP.h"
 #include "DerivedDataCacheInterface.h"
 #include "PackageName.h"
-
+#include "UniquePtr.h"
 
 #include "HTTPTransport.h"
 #include "TCPTransport.h"
 
+#include "UniquePtr.h"
+
 #if WITH_UNREAL_DEVELOPER_TOOLS
 	#include "Developer/PackageDependencyInfo/Public/PackageDependencyInfo.h"
 #endif	//WITH_UNREAL_DEVELOPER_TOOLS
@@ -895,8 +897,8 @@ public:
 			FString TempFilename = Filename + TEXT(".tmp");
 			InnerPlatformFile.CreateDirectoryTree(*FPaths::GetPath(Filename));
 			{
-				TAutoPtr<IFileHandle> FileHandle;
-				FileHandle = InnerPlatformFile.OpenWrite(*TempFilename);
+				TUniquePtr<IFileHandle> FileHandle;
+				FileHandle.Reset(InnerPlatformFile.OpenWrite(*TempFilename));
 
 				if (!FileHandle)
 				{
@@ -1053,11 +1055,11 @@ bool FNetworkPlatformFile::IsAdditionalCookedFileExtension(const TCHAR* Ext)
 {
 	if (*Ext != TEXT('.'))
 	{
-		return BulkFileExtension.EndsWith(Ext) || ExpFileExtension.EndsWith(Ext) || FontFileExtension.EndsWith(Ext);
+		return BulkFileExtension.EndsWith(Ext) || FontFileExtension.EndsWith(Ext) || ExpFileExtension.EndsWith(Ext);
 	}
 	else
 	{
-		return BulkFileExtension == Ext || ExpFileExtension == Ext || FontFileExtension == Ext;
+		return BulkFileExtension == Ext || FontFileExtension == Ext || ExpFileExtension == Ext;
 	}
 }
 
@@ -1302,8 +1304,8 @@ class FNetworkFileModule : public IPlatformFileModule
 public:
 	virtual IPlatformFile* GetPlatformFile() override
 	{
-		static TScopedPointer<IPlatformFile> AutoDestroySingleton(new FNetworkPlatformFile());
-		return AutoDestroySingleton.GetOwnedPointer();
+		static TUniquePtr<IPlatformFile> AutoDestroySingleton = MakeUnique<FNetworkPlatformFile>();
+		return AutoDestroySingleton.Get();
 	}
 };
 IMPLEMENT_MODULE(FNetworkFileModule, NetworkFile);
diff --git a/Engine/Source/Runtime/Online/BuildPatchServices/Private/BuildPatchGeneration.cpp b/Engine/Source/Runtime/Online/BuildPatchServices/Private/BuildPatchGeneration.cpp
--- a/Engine/Source/Runtime/Online/BuildPatchServices/Private/BuildPatchGeneration.cpp
+++ b/Engine/Source/Runtime/Online/BuildPatchServices/Private/BuildPatchGeneration.cpp
@@ -7,6 +7,7 @@
 #include "Generation/CloudEnumeration.h"
 #include "Generation/ManifestBuilder.h"
 #include "Generation/FileAttributesParser.h"
+#include "UniquePtr.h"
 
 using namespace BuildPatchServices;
 
@@ -157,7 +158,7 @@ bool FBuildDataGenerator::GenerateChunksManifestFromDirectory(const BuildPatchSe
 	auto* StatTotalTime = StatsCollector->CreateStat(TEXT("Generation: Total Time"), EStatFormat::Timer);
 
 	// List of created scanners.
-	TArray<TAutoPtr<FScannerDetails>> Scanners;
+	TArray<TUniquePtr<FScannerDetails>> Scanners;
 
 	// Tracking info per layer for rescanning.
 	TMap<int32, FChunkMatch> LayerToLastChunkMatch;
diff --git a/Engine/Source/Runtime/Online/BuildPatchServices/Private/Generation/FileAttributesParser.cpp b/Engine/Source/Runtime/Online/BuildPatchServices/Private/Generation/FileAttributesParser.cpp
--- a/Engine/Source/Runtime/Online/BuildPatchServices/Private/Generation/FileAttributesParser.cpp
+++ b/Engine/Source/Runtime/Online/BuildPatchServices/Private/Generation/FileAttributesParser.cpp
@@ -1,6 +1,7 @@
 // Copyright 1998-2016 Epic Games, Inc. All Rights Reserved.
 #include "BuildPatchServicesPrivatePCH.h"
 #include "FileAttributesParser.h"
+#include "UniquePtr.h"
 
 DECLARE_LOG_CATEGORY_EXTERN(LogFileAttributesParser, Log, All);
 DEFINE_LOG_CATEGORY(LogFileAttributesParser);
@@ -47,8 +48,8 @@ namespace BuildPatchServices
 
 	bool FFileAttributesParserImpl::ParseFileAttributes(const FString& MetaFilename, TMap<FString, FFileAttributes>& FileAttributes)
 	{
-		TAutoPtr<IFileHandle> Handle(PlatformFile.OpenRead(*MetaFilename));
-		if (Handle.IsValid())
+		TUniquePtr<IFileHandle> Handle(PlatformFile.OpenRead(*MetaFilename));
+		if (Handle)
 		{
 			TArray<uint8> FileData;
 			FileData.AddUninitialized(Handle->Size());
diff --git a/Engine/Source/Runtime/Online/SSL/Private/SslCertificateManager.cpp b/Engine/Source/Runtime/Online/SSL/Private/SslCertificateManager.cpp
--- a/Engine/Source/Runtime/Online/SSL/Private/SslCertificateManager.cpp
+++ b/Engine/Source/Runtime/Online/SSL/Private/SslCertificateManager.cpp
@@ -57,38 +57,38 @@ void FSslCertificateManager::BuildRootCertificateArray()
 	}
 
 	{
-		int64 CertificateBundleBufferSize = 0;
+	int64 CertificateBundleBufferSize = 0;
 		TUniquePtr<char[]> CertificateBundleBuffer;
 
 		if (TUniquePtr<FArchive> CertificateBundleArchive = TUniquePtr<FArchive>(IFileManager::Get().CreateFileReader(*CertificateBundlePath, 0)))
-		{
-			CertificateBundleBufferSize = CertificateBundleArchive->TotalSize();
+	{
+		CertificateBundleBufferSize = CertificateBundleArchive->TotalSize();
 			CertificateBundleBuffer.Reset(new char[CertificateBundleBufferSize + 1]);
 			CertificateBundleArchive->Serialize(CertificateBundleBuffer.Get(), CertificateBundleBufferSize);
 			CertificateBundleBuffer[CertificateBundleBufferSize] = '\0';
-		}
+	}
 
-		if (CertificateBundleBufferSize > 0 && CertificateBundleBuffer != nullptr)
-		{
-			static const char BeginCertificateString[] = "-----BEGIN CERTIFICATE-----";
-			static const char EndCertificateString[] = "-----END CERTIFICATE-----";
+	if (CertificateBundleBufferSize > 0 && CertificateBundleBuffer != nullptr)
+	{
+		static const char BeginCertificateString[] = "-----BEGIN CERTIFICATE-----";
+		static const char EndCertificateString[] = "-----END CERTIFICATE-----";
 
 			const char* FoundString = CertificateBundleBuffer.Get();
-			while (nullptr != (FoundString = FPlatformString::Strstr(FoundString, BeginCertificateString)))
+		while (nullptr != (FoundString = FPlatformString::Strstr(FoundString, BeginCertificateString)))
+		{
+			const char* EndString = FPlatformString::Strstr(FoundString, EndCertificateString);
+			if (EndString != nullptr)
 			{
-				const char* EndString = FPlatformString::Strstr(FoundString, EndCertificateString);
-				if (EndString != nullptr)
-				{
-					size_t LengthOfCertificateData = EndString - FoundString + sizeof(EndCertificateString);
-					BIO* CertificateBio = BIO_new_mem_buf(FoundString, LengthOfCertificateData);
-					X509* Certificate = PEM_read_bio_X509(CertificateBio, NULL, 0, NULL);
-					RootCertificateArray.Add(Certificate);
-					BIO_free(CertificateBio);
-				}
-				FoundString = EndString;
+				size_t LengthOfCertificateData = EndString - FoundString + sizeof(EndCertificateString);
+				BIO* CertificateBio = BIO_new_mem_buf(FoundString, LengthOfCertificateData);
+				X509* Certificate = PEM_read_bio_X509(CertificateBio, NULL, 0, NULL);
+				RootCertificateArray.Add(Certificate);
+				BIO_free(CertificateBio);
 			}
+			FoundString = EndString;
 		}
 	}
+	}
 
 	FString DebuggingCertificatePath;
 	if (GConfig->GetString(TEXT("SSL"), TEXT("DebuggingCertificatePath"), DebuggingCertificatePath, GEngineIni) && DebuggingCertificatePath.Len() > 0)
@@ -97,8 +97,9 @@ void FSslCertificateManager::BuildRootCertificateArray()
 		{
 			FArchive* DebuggingCertificateArchive = IFileManager::Get().CreateFileReader(*DebuggingCertificatePath, 0);
 			int64 CertificateBufferSize = DebuggingCertificateArchive->TotalSize();
-			char* CertificateBuffer = new char[CertificateBufferSize];
+			char* CertificateBuffer = new char[CertificateBufferSize + 1];
 			DebuggingCertificateArchive->Serialize(CertificateBuffer, CertificateBufferSize);
+			CertificateBuffer[CertificateBufferSize] = '\0';
 			BIO* CertificateBio = BIO_new_mem_buf(CertificateBuffer, -1);
 			X509* Certificate = PEM_read_bio_X509(CertificateBio, NULL, 0, NULL);
 			RootCertificateArray.Add(Certificate);
diff --git a/Engine/Source/Runtime/PakFile/Private/IPlatformFilePak.cpp b/Engine/Source/Runtime/PakFile/Private/IPlatformFilePak.cpp
--- a/Engine/Source/Runtime/PakFile/Private/IPlatformFilePak.cpp
+++ b/Engine/Source/Runtime/PakFile/Private/IPlatformFilePak.cpp
@@ -19,6 +19,17 @@ DEFINE_LOG_CATEGORY(LogPakFile);
 DEFINE_STAT(STAT_PakFile_Read);
 DEFINE_STAT(STAT_PakFile_NumOpenHandles);
 
+TPakChunkHash ComputePakChunkHash(const void* InData, int64 InDataSizeInBytes)
+{
+#if PAKHASH_USE_CRC
+	return FCrc::MemCrc32(InData, InDataSizeInBytes);
+#else
+	FSHAHash Hash;
+	FSHA1::HashBuffer(InData, InDataSizeInBytes, Hash);
+	return Hash;
+#endif
+}
+
 #ifndef EXCLUDE_NONPAK_UE_EXTENSIONS
 #define EXCLUDE_NONPAK_UE_EXTENSIONS USE_NEW_ASYNC_IO	// Use .Build.cs file to disable this if the game relies on accessing loose files
 #endif
@@ -39,11 +50,24 @@ FFilenameSecurityDelegate& FPakPlatformFile::GetFilenameSecurityDelegate()
 * Precaching
 */
 
+#ifdef AES_KEY
+DECLARE_DWORD_ACCUMULATOR_STAT(TEXT("PakCache Sync Decrypts (Uncompressed Path)"), STAT_PakCache_SyncDecrypts, STATGROUP_PakFile);
+DECLARE_FLOAT_ACCUMULATOR_STAT(TEXT("PakCache Decrypt Time"), STAT_PakCache_DecryptTime, STATGROUP_PakFile);
+DECLARE_DWORD_ACCUMULATOR_STAT(TEXT("PakCache Async Decrypts (Compressed Path)"), STAT_PakCache_CompressedDecrypts, STATGROUP_PakFile);
+DECLARE_DWORD_ACCUMULATOR_STAT(TEXT("PakCache Async Decrypts (Uncompressed Path)"), STAT_PakCache_UncompressedDecrypts, STATGROUP_PakFile);
+
+inline void DecryptData(uint8* InData, uint32 InDataSize)
+{
+	SCOPE_SECONDS_ACCUMULATOR(STAT_PakCache_DecryptTime);
+	FAES::DecryptData(InData, InDataSize);
+}
+#endif
+
 #if USE_PAK_PRECACHE
 #include "TaskGraphInterfaces.h"
 
 
-#define PAK_CACHE_GRANULARITY (64 * 1024)
+#define PAK_CACHE_GRANULARITY FPakInfo::MaxChunkDataSize
 #define PAK_CACHE_MAX_REQUESTS (8)
 #define PAK_CACHE_MAX_PRIORITY_DIFFERENCE_MERGE (AIOP_Normal - AIOP_Precache)
 #define PAK_EXTRA_CHECKS DO_CHECK
@@ -51,6 +75,10 @@ FFilenameSecurityDelegate& FPakPlatformFile::GetFilenameSecurityDelegate()
 DECLARE_MEMORY_STAT(TEXT("PakCache Current"), STAT_PakCacheMem, STATGROUP_Memory);
 DECLARE_MEMORY_STAT(TEXT("PakCache High Water"), STAT_PakCacheHighWater, STATGROUP_Memory);
 
+DECLARE_FLOAT_ACCUMULATOR_STAT(TEXT("PakCache Signing Chunk Hash Time"), STAT_PakCache_SigningChunkHashTime, STATGROUP_PakFile);
+DECLARE_MEMORY_STAT(TEXT("PakCache Signing Chunk Hash Size"), STAT_PakCache_SigningChunkHashSize, STATGROUP_PakFile);
+
+
 static int32 GPakCache_Enable = 1;
 static FAutoConsoleVariableRef CVar_Enable(
 	TEXT("pakcache.Enable"),
@@ -824,6 +852,8 @@ class FPakPrecacher
 		TIntervalTreeIndex InRequests[AIOP_NUM][(int32)EInRequestStatus::Num];
 		TIntervalTreeIndex CacheBlocks[(int32)EBlockStatus::Num];
 
+		TArray<TPakChunkHash> ChunkHashes;
+
 		FPakData(IAsyncReadFileHandle* InHandle, FName InName, int64 InTotalSize)
 			: Handle(InHandle)
 			, TotalSize(InTotalSize)
@@ -893,9 +923,13 @@ class FPakPrecacher
 	{
 		IAsyncReadRequest* RequestHandle;
 		TIntervalTreeIndex BlockIndex;
+		int64 RequestSize;
+		uint8* Memory;
 		FRequestToLower()
 			: RequestHandle(nullptr)
 			, BlockIndex(IntervalTreeInvalidIndex)
+			, RequestSize(0)
+			, Memory(nullptr)
 		{
 		}
 	};
@@ -908,14 +942,16 @@ class FPakPrecacher
 	uint32 Loads;
 	uint32 Frees;
 	uint64 LoadSize;
+	FEncryptionKey EncryptionKey;
+	bool bSigned;
 
 public:
 
-	static void Init(IPlatformFile* InLowerLevel)
+	static void Init(IPlatformFile* InLowerLevel, const FEncryptionKey& InEncryptionKey)
 	{
 		if (!PakPrecacherSingleton)
 		{
-			verify(!FPlatformAtomics::InterlockedCompareExchangePointer((void**)&PakPrecacherSingleton, new FPakPrecacher(InLowerLevel), nullptr));
+			verify(!FPlatformAtomics::InterlockedCompareExchangePointer((void**)&PakPrecacherSingleton, new FPakPrecacher(InLowerLevel, InEncryptionKey), nullptr));
 		}
 		check(PakPrecacherSingleton);
 	}
@@ -949,7 +985,7 @@ public:
 		return *PakPrecacherSingleton;
 	}
 
-	FPakPrecacher(IPlatformFile* InLowerLevel)
+	FPakPrecacher(IPlatformFile* InLowerLevel, const FEncryptionKey& InEncryptionKey)
 		: LowerLevel(InLowerLevel)
 		, LastReadRequest(0)
 		, NextUniqueID(1)
@@ -959,17 +995,41 @@ public:
 		, Loads(0)
 		, Frees(0)
 		, LoadSize(0)
+		, EncryptionKey(InEncryptionKey)
+		, bSigned(!InEncryptionKey.Exponent.IsZero() && !InEncryptionKey.Modulus.IsZero())
 	{
 		check(LowerLevel && FPlatformProcess::SupportsMultithreading());
 		GPakCache_MaxRequestsToLowerLevel = FMath::Max(FMath::Min(FPlatformMisc::NumberOfIOWorkerThreadsToSpawn(), GPakCache_MaxRequestsToLowerLevel), 1);
 		check(GPakCache_MaxRequestsToLowerLevel <= PAK_CACHE_MAX_REQUESTS);
 		CallbackFromLower = 
 			[this](bool bWasCanceled, IAsyncReadRequest* Request)
-			{				
-				NewRequestsToLowerComplete(bWasCanceled, Request);
+			{
+				{
+					FScopeLock Lock(&CachedFilesScopeLock);
+					for (int32 Index = 0; Index < PAK_CACHE_MAX_REQUESTS; Index++) // we loop over them all in case GPakCache_MaxRequestsToLowerLevel is changing
+					{
+						if (RequestsToLower[Index].RequestHandle == Request)
+						{
+							RequestsToLower[Index].Memory = Request->GetReadResults();
+							break;
+						}
+					}
+				}
+
+				if (bSigned)
+				{
+					StartSignatureCheck(bWasCanceled, Request);
+				}
+				else
+				{
+					NewRequestsToLowerComplete(bWasCanceled, Request);
+				}
 			};
 	}
 
+	void StartSignatureCheck(bool bWasCanceled, IAsyncReadRequest* Request);
+	void DoSignatureCheck(bool bWasCanceled, IAsyncReadRequest* Request);
+
 	IPlatformFile* GetLowerLevelHandle()
 	{
 		check(LowerLevel);
@@ -981,6 +1041,52 @@ public:
 		return GPakCache_AcceptPrecacheRequests;
 	}
 
+	uint16* RegisterPakFile(FName File, int64 PakFileSize)
+	{
+		uint16* PakIndexPtr = CachedPaks.Find(File);
+		if (!PakIndexPtr)
+		{
+			check(CachedPakData.Num() < MAX_uint16);
+			IAsyncReadFileHandle* Handle = LowerLevel->OpenAsyncRead(*File.ToString());
+			if (!Handle)
+			{
+				return nullptr;
+			}
+			CachedPakData.Add(FPakData(Handle, File, PakFileSize));
+			PakIndexPtr = &CachedPaks.Add(File, CachedPakData.Num() - 1);
+			UE_LOG(LogPakFile, Log, TEXT("New pak file %s added to pak precacher."), *File.ToString());
+
+			FPakData& Pak = CachedPakData[*PakIndexPtr];
+
+			if (bSigned)
+			{
+				// Load signature data
+				FString SignaturesFilename = FPaths::ChangeExtension(File.ToString(), TEXT("sig"));
+				IFileHandle* SignaturesFile = LowerLevel->OpenRead(*SignaturesFilename);
+				ensure(SignaturesFile);
+
+				FArchiveFileReaderGeneric* Reader = new FArchiveFileReaderGeneric(SignaturesFile, *SignaturesFilename, SignaturesFile->Size());
+				FEncryptedSignature MasterSignature;
+				*Reader << MasterSignature;
+				*Reader << Pak.ChunkHashes;
+				delete Reader;
+
+				// Check that we have the correct match between signature and pre-cache granularity
+				int64 NumPakChunks = Align(PakFileSize, PAK_CACHE_GRANULARITY) / PAK_CACHE_GRANULARITY;
+				ensure(NumPakChunks == Pak.ChunkHashes.Num());
+
+				// Decrypt signature hash
+				FDecryptedSignature DecryptedSignature;
+				FEncryption::DecryptSignature(MasterSignature, DecryptedSignature, EncryptionKey);
+
+				// Check the signatures are still as we expected them
+				TPakChunkHash Hash = ComputePakChunkHash(&Pak.ChunkHashes[0], Pak.ChunkHashes.Num() * sizeof(TPakChunkHash));
+				ensure(Hash == DecryptedSignature.Data);
+			}
+		}
+		return PakIndexPtr;
+	}
+
 private: // below here we assume CachedFilesScopeLock until we get to the next section
 
 	uint16 GetRequestPakIndex(FJoinedOffsetAndPakIndex OffsetAndPakIndex)
@@ -1669,6 +1775,8 @@ private: // below here we assume CachedFilesScopeLock until we get to the next s
 		uint16 PakIndex = GetRequestPakIndex(Block.OffsetAndPakIndex);
 		FPakData& Pak = CachedPakData[PakIndex];
 		RequestsToLower[IndexToFill].BlockIndex = Block.Index;
+		RequestsToLower[IndexToFill].RequestSize = Block.Size;
+		RequestsToLower[IndexToFill].Memory = nullptr;
 		check(&CacheBlockAllocator.Get(RequestsToLower[IndexToFill].BlockIndex) == &Block);
 
 		RequestsToLower[IndexToFill].RequestHandle = Pak.Handle->ReadRequest(GetRequestOffset(Block.OffsetAndPakIndex), Block.Size, Priority, &CallbackFromLower);
@@ -1677,7 +1785,7 @@ private: // below here we assume CachedFilesScopeLock until we get to the next s
 		LoadSize += Block.Size;
 	}
 
-	void CompleteRequest(bool bWasCanceled, IAsyncReadRequest* Request, TIntervalTreeIndex BlockIndex)
+	void CompleteRequest(bool bWasCanceled, uint8* Memory, TIntervalTreeIndex BlockIndex)
 	{
 		FCacheBlock& Block = CacheBlockAllocator.Get(BlockIndex);
 		uint16 PakIndex = GetRequestPakIndex(Block.OffsetAndPakIndex);
@@ -1697,7 +1805,7 @@ private: // below here we assume CachedFilesScopeLock until we get to the next s
 		}
 		else
 		{
-			Block.Memory = Request->GetReadResults();
+			Block.Memory = Memory;
 			check(Block.Memory && Block.Size);
 			BlockMemory += Block.Size;
 			check(BlockMemory > 0);
@@ -1858,6 +1966,8 @@ private: // below here we assume CachedFilesScopeLock until we get to the next s
 
 	///// Below here are the thread entrypoints
 
+public:
+
 	void NewRequestsToLowerComplete(bool bWasCanceled, IAsyncReadRequest* Request)
 	{
 		FScopeLock Lock(&CachedFilesScopeLock);
@@ -1867,7 +1977,7 @@ private: // below here we assume CachedFilesScopeLock until we get to the next s
 		{
 			if (RequestsToLower[Index].RequestHandle == Request)
 			{
-				CompleteRequest(bWasCanceled, Request, RequestsToLower[Index].BlockIndex);
+				CompleteRequest(bWasCanceled, RequestsToLower[Index].Memory, RequestsToLower[Index].BlockIndex);
 				RequestsToLower[Index].RequestHandle = nullptr;
 				RequestsToDelete.Add(Request);
 				RequestsToLower[Index].BlockIndex = IntervalTreeInvalidIndex;
@@ -1881,24 +1991,14 @@ private: // below here we assume CachedFilesScopeLock until we get to the next s
 		check(0); // not found?
 	}
 
-public:
-
 	bool QueueRequest(IPakRequestor* Owner, FName File, int64 PakFileSize, int64 Offset, int64 Size, EAsyncIOPriority Priority)
 	{
 		check(Owner && File != NAME_None && Size > 0 && Offset >= 0 && Offset < PakFileSize && Priority >= AIOP_MIN && Priority <= AIOP_MAX);
 		FScopeLock Lock(&CachedFilesScopeLock);
-		uint16* PakIndexPtr = CachedPaks.Find(File);
-		if (!PakIndexPtr)
+		uint16* PakIndexPtr = RegisterPakFile(File, PakFileSize);
+		if (PakIndexPtr == nullptr)
 		{
-			check(CachedPakData.Num() < MAX_uint16);
-			IAsyncReadFileHandle* Handle = LowerLevel->OpenAsyncRead(*File.ToString());
-			if (!Handle)
-			{
-				return false;
-			}
-			CachedPakData.Add(FPakData(Handle, File, PakFileSize));
-			PakIndexPtr = &CachedPaks.Add(File, CachedPakData.Num() - 1);
-			UE_LOG(LogPakFile, Log, TEXT("New pak file %s added to pak precacher."), *File.ToString());
+			return false;
 		}
 		uint16 PakIndex = *PakIndexPtr;
 		FPakData& Pak = CachedPakData[PakIndex];
@@ -2182,8 +2282,10 @@ static FCriticalSection FPakReadRequestEvent;
 class FPakAsyncReadFileHandle;
 //uncompress(unencrypt(checksig())))
 
-class FPakReadRequest : public IAsyncReadRequest, public IPakRequestor
+class FPakReadRequestBase : public IAsyncReadRequest, public IPakRequestor
 {
+protected:
+
 	int64 Offset;
 	int64 BytesToRead;
 	FEvent* WaitEvent;
@@ -2194,7 +2296,7 @@ class FPakReadRequest : public IAsyncReadRequest, public IPakRequestor
 	bool bInternalRequest; // we are using this internally to deal with compressed, encrypted and signed, so we want the memory back from a precache request.
 
 public:
-	FPakReadRequest(FName InPakFile, int64 PakFileSize, FAsyncFileCallBack* CompleteCallback, int64 InOffset, int64 InBytesToRead, EAsyncIOPriority InPriority, uint8* UserSuppliedMemory, bool bInInternalRequest = false, int32 InBlockIndex = -1)
+	FPakReadRequestBase(FName InPakFile, int64 PakFileSize, FAsyncFileCallBack* CompleteCallback, int64 InOffset, int64 InBytesToRead, EAsyncIOPriority InPriority, uint8* UserSuppliedMemory, bool bInInternalRequest = false, int32 InBlockIndex = -1)
 		: IAsyncReadRequest(CompleteCallback, false, UserSuppliedMemory)
 		, Offset(InOffset)
 		, BytesToRead(InBytesToRead)
@@ -2205,16 +2307,9 @@ public:
 		, bNeedsRemoval(true)
 		, bInternalRequest(bInInternalRequest)
 	{
-		check(Offset >= 0 && BytesToRead > 0);
-		check(bInternalRequest || Priority > AIOP_Precache || !bUserSuppliedMemory); // you never get bits back from a precache request, so why supply memory?
-		if (!FPakPrecacher::Get().QueueRequest(this, InPakFile, PakFileSize, Offset, BytesToRead, Priority))
-		{
-			bRequestOutstanding = false;
-			SetComplete();
-		}
 	}
 
-	virtual ~FPakReadRequest()
+	virtual ~FPakReadRequestBase()
 	{
 		if (bNeedsRemoval)
 		{
@@ -2274,6 +2369,29 @@ public:
 
 	// IPakRequestor Interface
 
+	int32 GetBlockIndex()
+	{
+		return BlockIndex;
+	}
+};
+
+class FPakReadRequest : public FPakReadRequestBase
+{
+public:
+
+	FPakReadRequest(FName InPakFile, int64 PakFileSize, FAsyncFileCallBack* CompleteCallback, int64 InOffset, int64 InBytesToRead, EAsyncIOPriority InPriority, uint8* UserSuppliedMemory, bool bInInternalRequest = false, int32 InBlockIndex = -1)
+		: FPakReadRequestBase(InPakFile, PakFileSize, CompleteCallback, InOffset, InBytesToRead, InPriority, UserSuppliedMemory, bInInternalRequest, InBlockIndex)
+	{
+		check(Offset >= 0 && BytesToRead > 0);
+		check(bInternalRequest || Priority > AIOP_Precache || !bUserSuppliedMemory); // you never get bits back from a precache request, so why supply memory?
+
+		if (!FPakPrecacher::Get().QueueRequest(this, InPakFile, PakFileSize, Offset, BytesToRead, Priority))
+		{
+			bRequestOutstanding = false;
+			SetComplete();
+		}
+	}
+
 	virtual void RequestIsComplete() override
 	{
 		check(bRequestOutstanding);
@@ -2302,11 +2420,85 @@ public:
 			WaitEvent->Trigger();
 		}
 	}
-	int32 GetBlockIndex()
+};
+
+#ifdef AES_KEY
+class FPakEncryptedReadRequest : public FPakReadRequestBase
+{
+	int64 OriginalOffset;
+	int64 OriginalSize;
+
+public:
+
+	FPakEncryptedReadRequest(FName InPakFile, int64 PakFileSize, FAsyncFileCallBack* CompleteCallback, int64 InPakFileStartOffset, int64 InFileOffset, int64 InBytesToRead, EAsyncIOPriority InPriority, uint8* UserSuppliedMemory, bool bInInternalRequest = false, int32 InBlockIndex = -1)
+		: FPakReadRequestBase(InPakFile, PakFileSize, CompleteCallback, InPakFileStartOffset + InFileOffset, InBytesToRead, InPriority, UserSuppliedMemory, bInInternalRequest, InBlockIndex)
+		, OriginalOffset(InPakFileStartOffset + InFileOffset)
+		, OriginalSize(InBytesToRead)
 	{
-		return BlockIndex;
+		Offset = InPakFileStartOffset + AlignDown(InFileOffset, FAES::AESBlockSize);
+		BytesToRead = Align(InBytesToRead, FAES::AESBlockSize);
+
+		if (!FPakPrecacher::Get().QueueRequest(this, InPakFile, PakFileSize, Offset, BytesToRead, Priority))
+		{
+			bRequestOutstanding = false;
+			SetComplete();
+		}
+	}
+
+	virtual void RequestIsComplete() override
+	{
+		check(bRequestOutstanding);
+		if (!bCanceled && (bInternalRequest || Priority > AIOP_Precache))
+		{
+			uint8* OversizedBuffer = nullptr;
+			if (OriginalOffset != Offset)
+			{
+				// We've read some bytes from before the requested offset, so we need to grab that larger amount
+				// from read request and then cut out the bit we want!
+				OversizedBuffer = (uint8*)FMemory::Malloc(BytesToRead);
+			}
+
+			if (!bUserSuppliedMemory)
+			{
+				check(!Memory);
+				Memory = (uint8*)FMemory::Malloc(OriginalSize);
+				INC_MEMORY_STAT_BY(STAT_AsyncFileMemory, OriginalSize);
+			}
+			else
+			{
+				check(Memory);
+			}
+
+			if (!FPakPrecacher::Get().GetCompletedRequest(this, OversizedBuffer != nullptr ? OversizedBuffer : Memory))
+			{
+				check(bCanceled);
+			}
+
+			INC_DWORD_STAT(STAT_PakCache_UncompressedDecrypts);
+
+			if (OversizedBuffer)
+			{
+				check(IsAligned((void*)BytesToRead, FAES::AESBlockSize));
+				DecryptData(OversizedBuffer, BytesToRead);
+				FMemory::Memcpy(Memory, OversizedBuffer + (OriginalOffset - Offset), OriginalSize);
+				FMemory::Free(OversizedBuffer);
+			}
+			else
+			{
+				DecryptData(Memory, Align(OriginalSize, FAES::AESBlockSize));
+			}
+		}
+
+		SetComplete();
+		FScopeLock Lock(&FPakReadRequestEvent);
+		bRequestOutstanding = false;
+		if (WaitEvent)
+		{
+			WaitEvent->Trigger();
+		}
 	}
 };
+#endif
 
 class FPakSizeRequest : public IAsyncReadRequest
 {
@@ -2486,6 +2678,115 @@ public:
 	void DoTask(ENamedThreads::Type CurrentThread, const FGraphEventRef& MyCompletionGraphEvent);
 };
 
+class FAsyncIOSignatureCheckTask
+{
+	bool bWasCanceled;
+	IAsyncReadRequest* Request;
+
+public:
+	FORCEINLINE FAsyncIOSignatureCheckTask(bool bInWasCanceled, IAsyncReadRequest* InRequest)
+		: bWasCanceled(bInWasCanceled)
+		, Request(InRequest)
+	{
+	}
+
+	static FORCEINLINE TStatId GetStatId()
+	{
+		RETURN_QUICK_DECLARE_CYCLE_STAT(FsyncIOCPUWorkTask, STATGROUP_TaskGraphTasks);
+	}
+	static FORCEINLINE ENamedThreads::Type GetDesiredThread()
+	{
+		return CPrio_AsyncIOCPUWorkTaskPriority.Get();
+	}
+	FORCEINLINE static ESubsequentsMode::Type GetSubsequentsMode()
+	{
+		return ESubsequentsMode::TrackSubsequents;
+	}
+	void DoTask(ENamedThreads::Type CurrentThread, const FGraphEventRef& MyCompletionGraphEvent)
+	{
+		FPakPrecacher::Get().DoSignatureCheck(bWasCanceled, Request);
+	}
+};
+
+void FPakPrecacher::StartSignatureCheck(bool bWasCanceled, IAsyncReadRequest* Request)
+{
+	FScopeLock Lock(&CachedFilesScopeLock);
+	for (int32 Index = 0; Index < PAK_CACHE_MAX_REQUESTS; Index++) // we loop over them all in case GPakCache_MaxRequestsToLowerLevel is changing
+	{
+		if (RequestsToLower[Index].RequestHandle == Request)
+		{
+			TGraphTask<FAsyncIOSignatureCheckTask>::CreateTask().ConstructAndDispatchWhenReady(bWasCanceled, Request);
+			return;
+		}
+	}
+	check(0); // not found?
+}
+
+void FPakPrecacher::DoSignatureCheck(bool bWasCanceled, IAsyncReadRequest* Request)
+{
+	int64 SignatureIndex = -1;
+	int64 NumSignaturesToCheck = 0;
+	const uint8* Data = nullptr;
+	int64 RequestSize = 0;
+	int64 RequestOffset = 0;
+	FPakData* PakData = nullptr;
+
+	{
+		// Try and keep lock for as short a time as possible. Find our request and copy out the data we need
+		FScopeLock Lock(&CachedFilesScopeLock);
+		for (int32 Index = 0; Index < PAK_CACHE_MAX_REQUESTS; Index++) // we loop over them all in case GPakCache_MaxRequestsToLowerLevel is changing
+		{
+			FRequestToLower& RequestToLower = RequestsToLower[Index];
+			if (RequestToLower.RequestHandle == Request)
+			{
+				NumSignaturesToCheck = Align(RequestToLower.RequestSize, PAK_CACHE_GRANULARITY) / PAK_CACHE_GRANULARITY;
+				check(NumSignaturesToCheck >= 1);
+
+				FCacheBlock& Block = CacheBlockAllocator.Get(RequestToLower.BlockIndex);
+				RequestOffset = GetRequestOffset(Block.OffsetAndPakIndex);
+				check((RequestOffset % PAK_CACHE_GRANULARITY) == 0);
+				RequestSize = RequestToLower.RequestSize;
+				int64 PakIndex = GetRequestPakIndex(Block.OffsetAndPakIndex);
+				PakData = &CachedPakData[PakIndex];
+				Data = RequestToLower.Memory;
+				SignatureIndex = RequestOffset / PAK_CACHE_GRANULARITY;
+				break;
+			}
+		}
+	}
+
+	check(Data);
+	check(NumSignaturesToCheck > 0);
+	check(RequestSize > 0);
+	check(RequestOffset >= 0);
+	check(PakData != nullptr);
+
+	// Hash the contents of the incoming buffer and check that it matches what we expected
+	for (int64 SignedChunkIndex = 0; SignedChunkIndex < NumSignaturesToCheck; ++SignedChunkIndex, ++SignatureIndex)
+	{
+		int64 Size = FMath::Min(RequestSize, (int64)PAK_CACHE_GRANULARITY);
+
+		{
+			SCOPE_SECONDS_ACCUMULATOR(STAT_PakCache_SigningChunkHashTime);
+			bool bChunkHashesMatch = ComputePakChunkHash(Data, Size) == PakData->ChunkHashes[SignatureIndex];
+			ensure(bChunkHashesMatch);
+			if (!ensure(bChunkHashesMatch))
+			{
+				UE_LOG(LogPakFile, Fatal, TEXT("Pak chunk signing mismatch! Pak file has been corrupted or tampered with!"));
+				FPlatformMisc::RequestExit(true);
+			}
+		}
+
+		INC_MEMORY_STAT_BY(STAT_PakCache_SigningChunkHashSize, Size);
+
+		RequestOffset += PAK_CACHE_GRANULARITY;
+		Data += PAK_CACHE_GRANULARITY;
+		RequestSize -= PAK_CACHE_GRANULARITY;
+	}
+
+	NewRequestsToLowerComplete(bWasCanceled, Request);
+}
+
 class FPakAsyncReadFileHandle final : public IAsyncReadFileHandle
 {
 	FName PakFile;
@@ -2549,7 +2850,19 @@ public:
 		{
 			check(Offset + BytesToRead + OffsetInPak <= PakFileSize);
 			check(!Blocks.Num());
-			return new FPakReadRequest(PakFile, PakFileSize, CompleteCallback, OffsetInPak + Offset, BytesToRead, Priority, UserSuppliedMemory);
+
+			if (FileEntry->bEncrypted)
+			{
+#ifdef AES_KEY
+				return new FPakEncryptedReadRequest(PakFile, PakFileSize, CompleteCallback, OffsetInPak, Offset, BytesToRead, Priority, UserSuppliedMemory);
+#else
+				UE_LOG(LogPakFile, Fatal, TEXT("Encountered an encrypted file in pak file but the executable wasn't compiled with the decryption key!"));
+#endif
+			}
+			else
+			{
+				return new FPakReadRequest(PakFile, PakFileSize, CompleteCallback, OffsetInPak + Offset, BytesToRead, Priority, UserSuppliedMemory);
+			}
 		}
 		bool bAnyUnfinished = false;
 		FPakProcessedReadRequest* Result;
@@ -2596,6 +2909,10 @@ public:
 		Block.bInFlight = true;
 		check(!Block.RawRequest && !Block.Processed && !Block.Raw && !Block.CPUWorkGraphEvent.GetReference() && !Block.ProcessedSize && !Block.RawSize);
 		Block.RawSize = FileEntry->CompressionBlocks[BlockIndex].CompressedEnd - FileEntry->CompressionBlocks[BlockIndex].CompressedStart;
+		if (FileEntry->bEncrypted)
+		{
+			Block.RawSize = Align(Block.RawSize, FAES::AESBlockSize);
+		}
 		NumLiveRawRequests++;
 		Block.RawRequest = new FPakReadRequest(PakFile, PakFileSize, &ReadCallbackFunction, FileEntry->CompressionBlocks[BlockIndex].CompressedStart, Block.RawSize, Priority, nullptr, true, BlockIndex);
 	}
@@ -2637,6 +2954,14 @@ public:
 		FCachedAsyncBlock& Block = Blocks[BlockIndex];
 		check(Block.Raw && Block.RawSize && !Block.Processed);
 
+#ifdef AES_KEY
+		if (FileEntry->bEncrypted)
+		{
+			INC_DWORD_STAT(STAT_PakCache_CompressedDecrypts);
+			DecryptData(Block.Raw, Align(Block.RawSize, FAES::AESBlockSize));
+		}
+#endif
+
 		check(Block.ProcessedSize);
 		INC_MEMORY_STAT_BY(STAT_AsyncFileMemory, Block.ProcessedSize);
 		uint8* Output = (uint8*)FMemory::Malloc(Block.ProcessedSize);
@@ -2884,7 +3209,8 @@ public:
 	static FORCEINLINE void DecryptBlock(void* Data, int64 Size)
 	{
 #ifdef AES_KEY
-		FAES::DecryptData((uint8*)Data, Size);
+		INC_DWORD_STAT(STAT_PakCache_SyncDecrypts);
+		DecryptData((uint8*)Data, Size);
 #endif
 	}
 };
@@ -2901,21 +3227,21 @@ public:
 	{}
 
 	int64				TempBufferSize;
-	TAutoPtr<uint8>		TempBuffer;
+	TUniquePtr<uint8[]>	TempBuffer;
 	int64				ScratchBufferSize;
-	TAutoPtr<uint8>		ScratchBuffer;
+	TUniquePtr<uint8[]>	ScratchBuffer;
 
 	void EnsureBufferSpace(int64 CompressionBlockSize, int64 ScrachSize)
 	{
 		if(TempBufferSize < CompressionBlockSize)
 		{
 			TempBufferSize = CompressionBlockSize;
-			TempBuffer.Reset((uint8*)FMemory::Malloc(TempBufferSize));
+			TempBuffer = MakeUnique<uint8[]>(TempBufferSize);
 		}
 		if(ScratchBufferSize < ScrachSize)
 		{
 			ScratchBufferSize = ScrachSize;
-			ScratchBuffer.Reset((uint8*)FMemory::Malloc(ScratchBufferSize));
+			ScratchBuffer = MakeUnique<uint8[]>(ScratchBufferSize);
 		}
 	}
 };
@@ -2991,8 +3317,8 @@ public:
 		int64 WorkingBufferRequiredSize = FCompression::CompressMemoryBound((ECompressionFlags)PakEntry.CompressionMethod,CompressionBlockSize, FPlatformMisc::GetPlatformCompression()->GetCompressionBitWindow());
 		WorkingBufferRequiredSize = EncryptionPolicy::AlignReadRequest(WorkingBufferRequiredSize);
 		ScratchSpace.EnsureBufferSpace(CompressionBlockSize, WorkingBufferRequiredSize*2);
-		WorkingBuffers[0] = ScratchSpace.ScratchBuffer;
-		WorkingBuffers[1] = ScratchSpace.ScratchBuffer + WorkingBufferRequiredSize;
+		WorkingBuffers[0] = ScratchSpace.ScratchBuffer.Get();
+		WorkingBuffers[1] = ScratchSpace.ScratchBuffer.Get() + WorkingBufferRequiredSize;
 
 		while (Length > 0)
 		{
@@ -3025,7 +3351,7 @@ public:
 			{
 				// Block needs to be copied from a working buffer
 				TaskDetails.Flags = (ECompressionFlags)PakEntry.CompressionMethod;
-				TaskDetails.UncompressedBuffer = (uint8*)ScratchSpace.TempBuffer;
+				TaskDetails.UncompressedBuffer = ScratchSpace.TempBuffer.Get();
 				TaskDetails.UncompressedSize = UncompressedBlockSize;
 				TaskDetails.CompressedBuffer = WorkingBuffers[CompressionBlockIndex & 1];
 				TaskDetails.CompressedSize = CompressedBlockSize;
@@ -3168,15 +3494,13 @@ FArchive* FPakFile::SetupSignedPakReader(FArchive* ReaderArchive, const TCHAR* F
 {
 	if (FPlatformProperties::RequiresCookedData())
 	{
-#if !USING_SIGNED_CONTENT
 		if (bSigned || FParse::Param(FCommandLine::Get(), TEXT("signedpak")) || FParse::Param(FCommandLine::Get(), TEXT("signed")))
-#endif
 		{
-			if (!Decryptor.IsValid())
+			if (!Decryptor)
 			{
-				Decryptor = new FChunkCacheWorker(ReaderArchive, Filename);
+				Decryptor = MakeUnique<FChunkCacheWorker>(ReaderArchive, Filename);
 			}
-			ReaderArchive = new FSignedArchiveReader(ReaderArchive, Decryptor);
+			ReaderArchive = new FSignedArchiveReader(ReaderArchive, Decryptor.Get());
 		}
 	}
 	return ReaderArchive;
@@ -3391,10 +3715,10 @@ FArchive* FPakFile::GetSharedReader(IPlatformFile* LowerLevel)
 	FArchive* PakReader = NULL;
 	{
 		FScopeLock ScopedLock(&CriticalSection);
-		TAutoPtr<FArchive>* ExistingReader = ReaderMap.Find(Thread);
+		TUniquePtr<FArchive>* ExistingReader = ReaderMap.Find(Thread);
 		if (ExistingReader)
 		{
-			PakReader = *ExistingReader;
+			PakReader = ExistingReader->Get();
 		}
 	}
 	if (!PakReader)
@@ -3423,7 +3747,7 @@ FArchive* FPakFile::GetSharedReader(IPlatformFile* LowerLevel)
 #else //DO_CHECK
 			ReaderMap.Emplace(Thread, PakReader);
 #endif //DO_CHECK
-		}		
+		}
 	}
 	return PakReader;
 }
@@ -3460,7 +3784,7 @@ public:
 		return false;
 	}
 };
-static TAutoPtr<FPakExec> GPakExec;
+static TUniquePtr<FPakExec> GPakExec;
 
 void FPakPlatformFile::HandleMountCommand(const TCHAR* Cmd, FOutputDevice& Ar)
 {
@@ -3613,18 +3937,14 @@ bool FPakPlatformFile::CheckIfPakFilesExist(IPlatformFile* LowLevelFile, const T
 
 bool FPakPlatformFile::ShouldBeUsed(IPlatformFile* Inner, const TCHAR* CmdLine) const
 {
-#if !USING_SIGNED_CONTENT
-	bool Result = FParse::Param(CmdLine, TEXT("Pak")) || FParse::Param(CmdLine, TEXT("Signedpak")) || FParse::Param(CmdLine, TEXT("Signed"));
-	if (FPlatformProperties::RequiresCookedData() && !Result && !FParse::Param(CmdLine, TEXT("NoPak")))
+	bool Result = false;
+	if (FPlatformProperties::RequiresCookedData() && !FParse::Param(CmdLine, TEXT("NoPak")))
 	{
 		TArray<FString> PakFolders;
 		GetPakFolders(CmdLine, PakFolders);
 		Result = CheckIfPakFilesExist(Inner, PakFolders);
 	}
 	return Result;
-#else
-	return true;
-#endif
 }
 
 bool FPakPlatformFile::Initialize(IPlatformFile* Inner, const TCHAR* CmdLine)
@@ -3641,19 +3961,10 @@ bool FPakPlatformFile::Initialize(IPlatformFile* Inner, const TCHAR* CmdLine)
 	ExcludedNonPakExtensions.Add(TEXT("uexp"));
 #endif
 
-#if !USING_SIGNED_CONTENT
-	bSigned = FParse::Param(CmdLine, TEXT("Signedpak")) || FParse::Param(CmdLine, TEXT("Signed"));
-	if (!bSigned)
-	{
-		// Even if -signed is not provided in the command line, use signed reader if the hardcoded key is non-zero.
-		FEncryptionKey DecryptionKey;
-		DecryptionKey.Exponent.Parse(DECRYPTION_KEY_EXPONENT);
-		DecryptionKey.Modulus.Parse(DECRYPTION_KEY_MODULUS);
-		bSigned = !DecryptionKey.Exponent.IsZero() && !DecryptionKey.Modulus.IsZero();
-	}
-#else
-	bSigned = true;
-#endif
+	FEncryptionKey DecryptionKey;
+	DecryptionKey.Exponent.Parse(DECRYPTION_KEY_EXPONENT);
+	DecryptionKey.Modulus.Parse(DECRYPTION_KEY_MODULUS);
+	bSigned = !DecryptionKey.Exponent.IsZero() && !DecryptionKey.Modulus.IsZero();
 	
 	bool bMountPaks = true;
 	TArray<FString> PaksToLoad;
@@ -3717,7 +4028,7 @@ bool FPakPlatformFile::Initialize(IPlatformFile* Inner, const TCHAR* CmdLine)
 	}
 
 #if !UE_BUILD_SHIPPING
-	GPakExec = new FPakExec(*this);
+	GPakExec = MakeUnique<FPakExec>(*this);
 #endif // !UE_BUILD_SHIPPING
 
 	FCoreDelegates::OnMountPak.BindRaw(this, &FPakPlatformFile::HandleMountPakDelegate);
@@ -3725,7 +4036,7 @@ bool FPakPlatformFile::Initialize(IPlatformFile* Inner, const TCHAR* CmdLine)
 #if USE_PAK_PRECACHE
 	if (!WITH_EDITOR && FPlatformProcess::SupportsMultithreading())
 	{
-		FPakPrecacher::Init(LowerLevel);
+		FPakPrecacher::Init(LowerLevel, DecryptionKey);
 	}
 	else
 	{
@@ -3869,7 +4180,6 @@ IFileHandle* FPakPlatformFile::OpenRead(const TCHAR* Filename, bool bAllowWrite)
 	{
 		Result = CreatePakFileHandle(Filename, PakFile, FileEntry);
 	}
-#if !USING_SIGNED_CONTENT
 	else
 	{
 		if (IsNonPakFilenameAllowed(Filename))
@@ -3878,7 +4188,6 @@ IFileHandle* FPakPlatformFile::OpenRead(const TCHAR* Filename, bool bAllowWrite)
 			Result = LowerLevel->OpenRead(Filename, bAllowWrite);
 		}
 	}
-#endif
 	return Result;
 }
 
@@ -3911,10 +4220,10 @@ bool FPakPlatformFile::CopyFile(const TCHAR* To, const TCHAR* From, EPlatformFil
 	{
 		// Copy from pak to LowerLevel->
 		// Create handles both files.
-		TAutoPtr<IFileHandle> DestHandle(LowerLevel->OpenWrite(To, false, (WriteFlags & EPlatformFileWrite::AllowRead) != EPlatformFileWrite::None));
-		TAutoPtr<IFileHandle> SourceHandle(CreatePakFileHandle(From, PakFile, FileEntry));
+		TUniquePtr<IFileHandle> DestHandle(LowerLevel->OpenWrite(To, false, (WriteFlags & EPlatformFileWrite::AllowRead) != EPlatformFileWrite::None));
+		TUniquePtr<IFileHandle> SourceHandle(CreatePakFileHandle(From, PakFile, FileEntry));
 
-		if (DestHandle.IsValid() && SourceHandle.IsValid())
+		if (DestHandle && SourceHandle)
 		{
 			const int64 BufferSize = 64 * 1024; // Copy in 64K chunks.
 			uint8* Buffer = (uint8*)FMemory::Malloc(BufferSize);
@@ -3929,11 +4238,6 @@ bool FPakPlatformFile::CopyFile(const TCHAR* To, const TCHAR* From, EPlatformFil
 	return Result;
 }
 
-uint32 ComputePakChunkHash(const uint8* InData, const int64 InDataSize)
-{
-	return FCrc::MemCrc32(InData, InDataSize);
-}
-
 /**
  * Module for the pak file
  */
@@ -3942,8 +4246,8 @@ class FPakFileModule : public IPlatformFileModule
 public:
 	virtual IPlatformFile* GetPlatformFile() override
 	{
-		static TScopedPointer<IPlatformFile> AutoDestroySingleton(new FPakPlatformFile());
-		return AutoDestroySingleton.GetOwnedPointer();
+		static TUniquePtr<IPlatformFile> AutoDestroySingleton = MakeUnique<FPakPlatformFile>();
+		return AutoDestroySingleton.Get();
 	}
 };
 
diff --git a/Engine/Source/Runtime/PakFile/Private/SignedArchiveReader.cpp b/Engine/Source/Runtime/PakFile/Private/SignedArchiveReader.cpp
--- a/Engine/Source/Runtime/PakFile/Private/SignedArchiveReader.cpp
+++ b/Engine/Source/Runtime/PakFile/Private/SignedArchiveReader.cpp
@@ -16,11 +16,6 @@ DECLARE_FLOAT_ACCUMULATOR_STAT(TEXT("FChunkCacheWorker.CheckSignature"), STAT_FC
 
 DECLARE_FLOAT_ACCUMULATOR_STAT(TEXT("FChunkCacheWorker.Serialize"), STAT_FChunkCacheWorker_Serialize, STATGROUP_PakFile);
 DECLARE_FLOAT_ACCUMULATOR_STAT(TEXT("FChunkCacheWorker.HashBuffer"), STAT_FChunkCacheWorker_HashBuffer, STATGROUP_PakFile);
-DECLARE_FLOAT_ACCUMULATOR_STAT(TEXT("FChunkCacheWorker.DecryptDuringIdle"), STAT_FChunkCacheWorker_DecryptDuringIdle, STATGROUP_PakFile);
-DECLARE_FLOAT_ACCUMULATOR_STAT(TEXT("FChunkCacheWorker.DecryptDuringVerify"), STAT_FChunkCacheWorker_DecryptDuringVerify, STATGROUP_PakFile);
-DECLARE_FLOAT_ACCUMULATOR_STAT(TEXT("FChunkCacheWorker.WaitForHashThread"), STAT_FChunkCacheWorker_WaitForHashThread, STATGROUP_PakFile);
-DECLARE_DWORD_ACCUMULATOR_STAT(TEXT("FChunkCacheWorker.SignaturesDecryptedDuringVerify"), STAT_FChunkCacheWorker_SignaturesDecryptedDuringVerify, STATGROUP_PakFile);
-DECLARE_DWORD_ACCUMULATOR_STAT(TEXT("FChunkCacheWorker.SignaturesDecryptedDuringIdle"), STAT_FChunkCacheWorker_SignaturesDecryptedDuringIdle, STATGROUP_PakFile);
 
 DECLARE_FLOAT_ACCUMULATOR_STAT(TEXT("FSignedArchiveReader.Serialize"), STAT_SignedArchiveReader_Serialize, STATGROUP_PakFile);
 
@@ -28,6 +23,8 @@ FChunkCacheWorker::FChunkCacheWorker(FArchive* InReader, const TCHAR* Filename)
 	: Reader(InReader)
 	, QueuedRequestsEvent(NULL)
 {
+	SetupDecryptionKey();
+
 	FString SigFileFilename = FPaths::ChangeExtension(Filename, TEXT("sig"));
 	FArchive* SigFileReader = IFileManager::Get().CreateFileReader(*SigFileFilename);
 
@@ -36,15 +33,15 @@ FChunkCacheWorker::FChunkCacheWorker(FArchive* InReader, const TCHAR* Filename)
 		UE_LOG(LogPakFile, Fatal, TEXT("Couldn't find pak signature file '%s'"), *Filename);
 	}
 
-	*SigFileReader << EncryptedSignatures;
+	FEncryptedSignature MasterSignature;
+	*SigFileReader << MasterSignature;
+	*SigFileReader << ChunkHashes;
 	delete SigFileReader;
 
-	EncryptedSignaturesCRC = FCrc::MemCrc32(&EncryptedSignatures[0], EncryptedSignatures.Num() * sizeof(FEncryptedSignature));
-	DecryptedSignaturesCRC = 0;
-
-	DecryptedSignatures.AddDefaulted(EncryptedSignatures.Num());
-
-	SetupDecryptionKey();
+	uint32 ChunkHashesCRC = ComputePakChunkHash(&ChunkHashes[0], ChunkHashes.Num() * sizeof(TPakChunkHash));
+	FDecryptedSignature DecryptedMasterSignature;
+	FEncryption::DecryptSignature(MasterSignature, DecryptedMasterSignature, DecryptionKey);
+	ensure(DecryptedMasterSignature.Data == ChunkHashesCRC);
 
 	QueuedRequestsEvent = FPlatformProcess::GetSynchEventFromPool();
 	Thread = FRunnableThread::Create(this, TEXT("FChunkCacheWorker"), 0, TPri_BelowNormal);
@@ -88,57 +85,13 @@ void FChunkCacheWorker::SetupDecryptionKey()
 	UE_CLOG(bIdentical, LogPakFile, Fatal, TEXT("Decryption key produces identical results to source data."));
 }
 
-void FChunkCacheWorker::DecryptSignatures(int32 NextIndexToDecrypt)
-{
-	FDecryptedSignature& DecryptedSignature = DecryptedSignatures[NextIndexToDecrypt];
-
-	// Check to see if this signature was already decrypted.
-	if (!DecryptedSignature.IsValid())
-	{
-		INC_DWORD_STAT(STAT_FChunkCacheWorker_SignaturesDecryptedDuringIdle);
-		SCOPE_SECONDS_ACCUMULATOR(STAT_FChunkCacheWorker_DecryptDuringIdle);
-		FEncryptedSignature& EncryptedSignature = EncryptedSignatures[NextIndexToDecrypt];
-		FEncryption::DecryptSignature(EncryptedSignature, DecryptedSignature, DecryptionKey);
-	}
-}
-
 uint32 FChunkCacheWorker::Run()
-{
-	int32 NextIndexToDecrypt = 0;
-	LastDecryptedSignatureIndex = -1;
-	StartTime = FPlatformTime::Seconds();
-	
+{	
 	while (StopTaskCounter.GetValue() == 0)
 	{
-		int32 ProcessedRequests = ProcessQueue();		
-		if (ProcessedRequests == 0)
+		if (ProcessQueue() == 0)
 		{
-			uint32 WaitTime = 500;
-
-			// Try and decrypt some signatures if there are any left
-			if (NextIndexToDecrypt < EncryptedSignatures.Num())
-			{
-				check(NextIndexToDecrypt < EncryptedSignatures.Num());
-
-				DecryptSignatures(NextIndexToDecrypt);
-				LastDecryptedSignatureIndex = NextIndexToDecrypt;
-
-				if (++NextIndexToDecrypt == EncryptedSignatures.Num())
-				{
-					double Time = FPlatformTime::Seconds() - StartTime;
-					UE_LOG(LogPakFile, Log, TEXT("PakFile signature decryption complete in %.2fs"), Time);
-
-					DecryptedSignaturesCRC = FCrc::MemCrc32(&DecryptedSignatures[0], DecryptedSignatures.Num() * sizeof(FDecryptedSignature));
-				}
-				else
-				{
-					// There are more signatures to decrypt, so don't stall on event wait if there are
-					// no pending chunk requests
-					WaitTime = 0;
-				}
-			}
-
-			QueuedRequestsEvent->Wait(WaitTime);
+			QueuedRequestsEvent->Wait(500);
 		}
 	}
 	return 0;
@@ -260,67 +213,33 @@ bool FChunkCacheWorker::CheckSignature(const FChunkRequest& ChunkInfo)
 {
 	SCOPE_SECONDS_ACCUMULATOR(STAT_FChunkCacheWorker_CheckSignature);
 
-	FDecryptedSignature SourceSignature;
-	const int32 MaxNumRetries = 3;
-	int32 RetriesRemaining = MaxNumRetries;
-	bool bPakChunkSignaturesMatched = false;
+	TPakChunkHash ChunkHash = 0;
 
-	while (!bPakChunkSignaturesMatched && RetriesRemaining >= 0)
 	{
-		{
-			SCOPE_SECONDS_ACCUMULATOR(STAT_FChunkCacheWorker_Serialize);
-			Reader->Seek(ChunkInfo.Offset);
-			Reader->Serialize(ChunkInfo.Buffer->Data, ChunkInfo.Size);
-		}
-		{
-			SCOPE_SECONDS_ACCUMULATOR(STAT_FChunkCacheWorker_HashBuffer);
-			SourceSignature.Data = FCrc::MemCrc32(ChunkInfo.Buffer->Data, ChunkInfo.Size);
-		}
-
-		// Decrypt serialized hash
-		FDecryptedSignature& DecryptedSignature = DecryptedSignatures[ChunkInfo.Index];
-
-		if (RetriesRemaining < MaxNumRetries || !DecryptedSignature.IsValid())
-		{
-			INC_DWORD_STAT(STAT_FChunkCacheWorker_SignaturesDecryptedDuringVerify);
-			SCOPE_SECONDS_ACCUMULATOR(STAT_FChunkCacheWorker_DecryptDuringVerify);
-			FEncryptedSignature EncryptedSignature = EncryptedSignatures[ChunkInfo.Index];
-			check(EncryptedSignature.IsValid());
-			FEncryption::DecryptSignature(EncryptedSignature, DecryptedSignature, DecryptionKey);
-			check(DecryptedSignature.IsValid());
-		}
-
-		// Compare hashes
-		bPakChunkSignaturesMatched = DecryptedSignature == SourceSignature;
-		--RetriesRemaining;
+		SCOPE_SECONDS_ACCUMULATOR(STAT_FChunkCacheWorker_Serialize);
+		Reader->Seek(ChunkInfo.Offset);
+		Reader->Serialize(ChunkInfo.Buffer->Data, ChunkInfo.Size);
 	}
-
-	if (!bPakChunkSignaturesMatched)
 	{
-		// Check that the encrypted signatures data is the same as it was when we started up
-		int32 CurrentEncryptedSignaturesCRC = FCrc::MemCrc32(&EncryptedSignatures[0], EncryptedSignatures.Num() * sizeof(FEncryptedSignature));
-		ensure(CurrentEncryptedSignaturesCRC == EncryptedSignaturesCRC);
-
-		// If we finished decrypting all the signatures, check that they are still the same as they were when decryption completed
-		if (DecryptedSignaturesCRC != 0)
-		{
-			int32 CurrentDecryptedSignaturesCRC = FCrc::MemCrc32(&DecryptedSignatures[0], DecryptedSignatures.Num() * sizeof(FDecryptedSignature));
-			ensure(CurrentDecryptedSignaturesCRC == DecryptedSignaturesCRC);
-		}
+		SCOPE_SECONDS_ACCUMULATOR(STAT_FChunkCacheWorker_HashBuffer);
+		ChunkHash = ComputePakChunkHash(ChunkInfo.Buffer->Data, ChunkInfo.Size);
+	}
 
+	bool bHashesMatch = (ChunkHash == ChunkHashes[ChunkInfo.Index]);
+	if (!bHashesMatch)
+	{
 		UE_LOG(LogPakFile, Warning, TEXT("Pak chunk signature verification failed!"));
 		UE_LOG(LogPakFile, Warning, TEXT("  Chunk Index: %d"), ChunkInfo.Index);
 		UE_LOG(LogPakFile, Warning, TEXT("  Chunk Offset: %d"), ChunkInfo.Offset);
 		UE_LOG(LogPakFile, Warning, TEXT("  Chunk Size: %d"), ChunkInfo.Size);
-		UE_LOG(LogPakFile, Warning, TEXT("  Background decrypt: %d\\%d"), LastDecryptedSignatureIndex, DecryptedSignatures.Num());
-
-		ensure(bPakChunkSignaturesMatched);
+	
+		ensure(bHashesMatch);
 
 		//FPlatformMisc::MessageBoxExt(EAppMsgType::Ok, TEXT("Corrupt Installation Detected. Please use \"Verify\" in the Epic Games Launcher"), TEXT("Pakfile Error"));
 		//FPlatformMisc::RequestExit(1);
 	}
 	
-	return bPakChunkSignaturesMatched;
+	return bHashesMatch;
 }
 
 FChunkRequest& FChunkCacheWorker::RequestChunk(int32 ChunkIndex, int64 StartOffset, int64 ChunkSize)
diff --git a/Engine/Source/Runtime/PakFile/Private/SignedArchiveReader.h b/Engine/Source/Runtime/PakFile/Private/SignedArchiveReader.h
--- a/Engine/Source/Runtime/PakFile/Private/SignedArchiveReader.h
+++ b/Engine/Source/Runtime/PakFile/Private/SignedArchiveReader.h
@@ -98,10 +98,8 @@ class FChunkCacheWorker : public FRunnable
 		MaxCachedChunks = 8		
 	};
 
-	/** Encrypted signatures */
-	TArray<FEncryptedSignature> EncryptedSignatures;
-	/** Decrypted signatures */
-	TArray<FDecryptedSignature> DecryptedSignatures;
+	/** Reference hashes */
+	TArray<uint32> ChunkHashes;
 	/** Thread to run the worker FRunnable on */
 	FRunnableThread* Thread;
 	/** Archive reader */
@@ -122,14 +120,6 @@ class FChunkCacheWorker : public FRunnable
 	TLockFreePointerListUnordered<FChunkRequest, PLATFORM_CACHE_LINE_SIZE> FreeChunkRequests;
 	/** Public decryption key */
 	FEncryptionKey DecryptionKey;
-	/** The time at which the chunk cache worker started, for debugging purposes */
-	double StartTime;
-	/** Track the index of the last signature decrypted in the background, for debugging purposes */
-	int32 LastDecryptedSignatureIndex;
-	/** Stored CRC for the encrypted signature data after loading */
-	uint32 EncryptedSignaturesCRC;
-	/** Stored CRC for the fully dencrypted signature data after loading. Only computed when all signatures were decrypted */
-	uint32 DecryptedSignaturesCRC;
 
 	/** 
 	 * Process requested chunks 
@@ -141,10 +131,6 @@ class FChunkCacheWorker : public FRunnable
 	 * Verifies chunk signature [*]
 	 */
 	bool CheckSignature(const FChunkRequest& ChunkInfo);
-	/** 
-	* Decrypts multiple signature [*]
-	*/
-	void DecryptSignatures(int32 NextIndexToDecrypt);
 	/** 
 	 * Tries to get a pre-cached chunk buffer 
 	 */
diff --git a/Engine/Source/Runtime/PakFile/Public/IPlatformFilePak.h b/Engine/Source/Runtime/PakFile/Public/IPlatformFilePak.h
--- a/Engine/Source/Runtime/PakFile/Public/IPlatformFilePak.h
+++ b/Engine/Source/Runtime/PakFile/Public/IPlatformFilePak.h
@@ -2,6 +2,8 @@
 
 #pragma once
 
+#include "UniquePtr.h"
+
 PAKFILE_API DECLARE_LOG_CATEGORY_EXTERN(LogPakFile, Log, All);
 DECLARE_FLOAT_ACCUMULATOR_STAT_EXTERN(TEXT("Total pak file read time"), STAT_PakFile_Read, STATGROUP_PakFile, PAKFILE_API);
 
@@ -10,6 +12,16 @@ DECLARE_DWORD_ACCUMULATOR_STAT_EXTERN(TEXT("Num open pak file handles"), STAT_Pa
 /** Delegate for allowing a game to restrict the accessing of non-pak files */
 DECLARE_DELEGATE_RetVal_OneParam(bool, FFilenameSecurityDelegate, const TCHAR* /*InFilename*/);
 
+#define PAKHASH_USE_CRC	1
+
+#if PAKHASH_USE_CRC
+typedef uint32 TPakChunkHash;
+#else
+typedef FSHAHash TPakChunkHash;
+#endif
+
+PAKFILE_API TPakChunkHash ComputePakChunkHash(const void* InData, int64 InDataSizeInBytes);
+
 /**
  * Struct which holds pak file info (version, index offset, hash value).
  */
@@ -20,7 +32,7 @@ struct FPakInfo
 		/** Magic number to use in header */
 		PakFile_Magic = 0x5A6F12E1,
 		/** Size of cached data. */
-		MaxChunkDataSize = 256*1024,
+		MaxChunkDataSize = 64*1024,
 	};
 
 	/** Version numbers. */
@@ -260,9 +272,9 @@ class PAKFILE_API FPakFile : FNoncopyable
 	FString PakFilename;
 	FName PakFilenameName;
 	/** Archive to serialize the pak file from. */
-	TAutoPtr<class FChunkCacheWorker> Decryptor;
+	TUniquePtr<class FChunkCacheWorker> Decryptor;
 	/** Map of readers assigned to threads. */
-	TMap<uint32, TAutoPtr<FArchive>> ReaderMap;
+	TMap<uint32, TUniquePtr<FArchive>> ReaderMap;
 	/** Critical section for accessing ReaderMap. */
 	FCriticalSection CriticalSection;
 	/** Pak file info (trailer). */
diff --git a/Engine/Source/Runtime/RHI/Private/RHI.cpp b/Engine/Source/Runtime/RHI/Private/RHI.cpp
--- a/Engine/Source/Runtime/RHI/Private/RHI.cpp
+++ b/Engine/Source/Runtime/RHI/Private/RHI.cpp
@@ -291,6 +291,7 @@ bool GSupportsParallelOcclusionQueries = false;
 bool GSupportsRenderTargetWriteMask = false;
 
 bool GRHISupportsMSAADepthSampleAccess = false;
+bool GRHISupportsResolveCubemapFaces = false;
 
 bool GRHISupportsHDROutput = false;
 
diff --git a/Engine/Source/Runtime/RHI/Public/RHI.h b/Engine/Source/Runtime/RHI/Public/RHI.h
--- a/Engine/Source/Runtime/RHI/Public/RHI.h
+++ b/Engine/Source/Runtime/RHI/Public/RHI.h
@@ -256,6 +256,9 @@ extern RHI_API bool GRHISupportsBaseVertexIndex;
 /** True if the RHI supports hardware instancing */
 extern RHI_API bool GRHISupportsInstancing;
 
+/** True if the RHI supports copying cubemap faces using CopyToResolveTarget */
+extern RHI_API bool GRHISupportsResolveCubemapFaces;
+
 /** Whether or not the RHI can handle a non-zero FirstInstance - extra SetStreamSource calls will be needed if this is false */
 extern RHI_API bool GRHISupportsFirstInstance;
 
