Copying //UE4/Dev-Core to //UE4/Dev-Main (Source: //UE4/Dev-Core @ 3283640)

diff --git a/Engine/Source/Runtime/PakFile/Private/IPlatformFilePak.cpp b/Engine/Source/Runtime/PakFile/Private/IPlatformFilePak.cpp
--- a/Engine/Source/Runtime/PakFile/Private/IPlatformFilePak.cpp
+++ b/Engine/Source/Runtime/PakFile/Private/IPlatformFilePak.cpp
@@ -959,7 +959,6 @@ class FPakPrecacher
 	FRequestToLower RequestsToLower[PAK_CACHE_MAX_REQUESTS];
 	TArray<IAsyncReadRequest*> RequestsToDelete;
 	int32 NotifyRecursion;
-	FAsyncFileCallBack CallbackFromLower;
 
 	uint32 Loads;
 	uint32 Frees;
@@ -1023,34 +1022,10 @@ public:
 		check(LowerLevel && FPlatformProcess::SupportsMultithreading());
 		GPakCache_MaxRequestsToLowerLevel = FMath::Max(FMath::Min(FPlatformMisc::NumberOfIOWorkerThreadsToSpawn(), GPakCache_MaxRequestsToLowerLevel), 1);
 		check(GPakCache_MaxRequestsToLowerLevel <= PAK_CACHE_MAX_REQUESTS);
-		CallbackFromLower = 
-			[this](bool bWasCanceled, IAsyncReadRequest* Request)
-			{
-				{
-					FScopeLock Lock(&CachedFilesScopeLock);
-					for (int32 Index = 0; Index < PAK_CACHE_MAX_REQUESTS; Index++) // we loop over them all in case GPakCache_MaxRequestsToLowerLevel is changing
-					{
-						if (RequestsToLower[Index].RequestHandle == Request)
-						{
-							RequestsToLower[Index].Memory = Request->GetReadResults();
-							break;
-						}
-					}
-				}
-
-				if (bSigned)
-				{
-					StartSignatureCheck(bWasCanceled, Request);
-				}
-				else
-				{
-					NewRequestsToLowerComplete(bWasCanceled, Request);
-				}
-			};
 	}
 
-	void StartSignatureCheck(bool bWasCanceled, IAsyncReadRequest* Request);
-	void DoSignatureCheck(bool bWasCanceled, IAsyncReadRequest* Request);
+	void StartSignatureCheck(bool bWasCanceled, IAsyncReadRequest* Request, int32 IndexToFill);
+	void DoSignatureCheck(bool bWasCanceled, IAsyncReadRequest* Request, int32 IndexToFill);
 
 	IPlatformFile* GetLowerLevelHandle()
 	{
@@ -1801,6 +1776,19 @@ private: // below here we assume CachedFilesScopeLock until we get to the next s
 		RequestsToLower[IndexToFill].Memory = nullptr;
 		check(&CacheBlockAllocator.Get(RequestsToLower[IndexToFill].BlockIndex) == &Block);
 
+		FAsyncFileCallBack CallbackFromLower =
+			[this, IndexToFill](bool bWasCanceled, IAsyncReadRequest* Request)
+		{
+			if (bSigned)
+			{
+				StartSignatureCheck(bWasCanceled, Request, IndexToFill);
+			}
+			else
+			{
+				NewRequestsToLowerComplete(bWasCanceled, Request, IndexToFill);
+			}
+		};
+
 		RequestsToLower[IndexToFill].RequestHandle = Pak.Handle->ReadRequest(GetRequestOffset(Block.OffsetAndPakIndex), Block.Size, Priority, &CallbackFromLower);
 		LastReadRequest = Block.OffsetAndPakIndex + Block.Size;
 		Loads++;
@@ -1839,7 +1827,7 @@ private: // below here we assume CachedFilesScopeLock until we get to the next s
 				BlockMemoryHighWater = BlockMemory;
 				SET_MEMORY_STAT(STAT_PakCacheHighWater, BlockMemoryHighWater);
 
-#if 1
+#if 0
 				static int64 LastPrint = 0;
 				if (BlockMemoryHighWater / 1024 / 1024 != LastPrint)
 				{
@@ -1990,27 +1978,22 @@ private: // below here we assume CachedFilesScopeLock until we get to the next s
 
 public:
 
-	void NewRequestsToLowerComplete(bool bWasCanceled, IAsyncReadRequest* Request)
+	void NewRequestsToLowerComplete(bool bWasCanceled, IAsyncReadRequest* Request, int32 Index)
 	{
 		FScopeLock Lock(&CachedFilesScopeLock);
+		RequestsToLower[Index].RequestHandle = Request;
 		ClearOldBlockTasks();
 		NotifyRecursion++;
-		for (int32 Index = 0; Index < PAK_CACHE_MAX_REQUESTS; Index++) // we loop over them all in case GPakCache_MaxRequestsToLowerLevel is changing
+		if (!RequestsToLower[Index].Memory) // might have already been filled in by the signature check
 		{
-			if (RequestsToLower[Index].RequestHandle == Request)
-			{
-				CompleteRequest(bWasCanceled, RequestsToLower[Index].Memory, RequestsToLower[Index].BlockIndex);
-				RequestsToLower[Index].RequestHandle = nullptr;
-				RequestsToDelete.Add(Request);
-				RequestsToLower[Index].BlockIndex = IntervalTreeInvalidIndex;
-				StartNextRequest();
-				NotifyRecursion--;
-				return;
-			}
+			RequestsToLower[Index].Memory = Request->GetReadResults();
 		}
+		CompleteRequest(bWasCanceled, RequestsToLower[Index].Memory, RequestsToLower[Index].BlockIndex);
+		RequestsToLower[Index].RequestHandle = nullptr;
+		RequestsToDelete.Add(Request);
+		RequestsToLower[Index].BlockIndex = IntervalTreeInvalidIndex;
 		StartNextRequest();
 		NotifyRecursion--;
-		check(0); // not found?
 	}
 
 	bool QueueRequest(IPakRequestor* Owner, FName File, int64 PakFileSize, int64 Offset, int64 Size, EAsyncIOPriority Priority)
@@ -2702,11 +2685,13 @@ class FAsyncIOSignatureCheckTask
 {
 	bool bWasCanceled;
 	IAsyncReadRequest* Request;
+	int32 IndexToFill;
 
 public:
-	FORCEINLINE FAsyncIOSignatureCheckTask(bool bInWasCanceled, IAsyncReadRequest* InRequest)
+	FORCEINLINE FAsyncIOSignatureCheckTask(bool bInWasCanceled, IAsyncReadRequest* InRequest, int32 InIndexToFill)
 		: bWasCanceled(bInWasCanceled)
 		, Request(InRequest)
+		, IndexToFill(InIndexToFill)
 	{
 	}
 
@@ -2724,25 +2709,16 @@ public:
 	}
 	void DoTask(ENamedThreads::Type CurrentThread, const FGraphEventRef& MyCompletionGraphEvent)
 	{
-		FPakPrecacher::Get().DoSignatureCheck(bWasCanceled, Request);
+		FPakPrecacher::Get().DoSignatureCheck(bWasCanceled, Request, IndexToFill);
 	}
 };
 
-void FPakPrecacher::StartSignatureCheck(bool bWasCanceled, IAsyncReadRequest* Request)
+void FPakPrecacher::StartSignatureCheck(bool bWasCanceled, IAsyncReadRequest* Request, int32 Index)
 {
-	FScopeLock Lock(&CachedFilesScopeLock);
-	for (int32 Index = 0; Index < PAK_CACHE_MAX_REQUESTS; Index++) // we loop over them all in case GPakCache_MaxRequestsToLowerLevel is changing
-	{
-		if (RequestsToLower[Index].RequestHandle == Request)
-		{
-			TGraphTask<FAsyncIOSignatureCheckTask>::CreateTask().ConstructAndDispatchWhenReady(bWasCanceled, Request);
-			return;
-		}
-	}
-	check(0); // not found?
+	TGraphTask<FAsyncIOSignatureCheckTask>::CreateTask().ConstructAndDispatchWhenReady(bWasCanceled, Request, Index);
 }
 
-void FPakPrecacher::DoSignatureCheck(bool bWasCanceled, IAsyncReadRequest* Request)
+void FPakPrecacher::DoSignatureCheck(bool bWasCanceled, IAsyncReadRequest* Request, int32 Index)
 {
 	int64 SignatureIndex = -1;
 	int64 NumSignaturesToCheck = 0;
@@ -2754,25 +2730,21 @@ void FPakPrecacher::DoSignatureCheck(bool bWasCanceled, IAsyncReadRequest* Reque
 	{
 		// Try and keep lock for as short a time as possible. Find our request and copy out the data we need
 		FScopeLock Lock(&CachedFilesScopeLock);
-		for (int32 Index = 0; Index < PAK_CACHE_MAX_REQUESTS; Index++) // we loop over them all in case GPakCache_MaxRequestsToLowerLevel is changing
-		{
-			FRequestToLower& RequestToLower = RequestsToLower[Index];
-			if (RequestToLower.RequestHandle == Request)
-			{
-				NumSignaturesToCheck = Align(RequestToLower.RequestSize, PAK_CACHE_GRANULARITY) / PAK_CACHE_GRANULARITY;
-				check(NumSignaturesToCheck >= 1);
-
-				FCacheBlock& Block = CacheBlockAllocator.Get(RequestToLower.BlockIndex);
-				RequestOffset = GetRequestOffset(Block.OffsetAndPakIndex);
-				check((RequestOffset % PAK_CACHE_GRANULARITY) == 0);
-				RequestSize = RequestToLower.RequestSize;
-				int64 PakIndex = GetRequestPakIndex(Block.OffsetAndPakIndex);
-				PakData = &CachedPakData[PakIndex];
-				Data = RequestToLower.Memory;
-				SignatureIndex = RequestOffset / PAK_CACHE_GRANULARITY;
-				break;
-			}
-		}
+		FRequestToLower& RequestToLower = RequestsToLower[Index];
+		RequestToLower.RequestHandle = Request;
+		RequestToLower.Memory = Request->GetReadResults();
+
+		NumSignaturesToCheck = Align(RequestToLower.RequestSize, PAK_CACHE_GRANULARITY) / PAK_CACHE_GRANULARITY;
+		check(NumSignaturesToCheck >= 1);
+
+		FCacheBlock& Block = CacheBlockAllocator.Get(RequestToLower.BlockIndex);
+		RequestOffset = GetRequestOffset(Block.OffsetAndPakIndex);
+		check((RequestOffset % PAK_CACHE_GRANULARITY) == 0);
+		RequestSize = RequestToLower.RequestSize;
+		int64 PakIndex = GetRequestPakIndex(Block.OffsetAndPakIndex);
+		PakData = &CachedPakData[PakIndex];
+		Data = RequestToLower.Memory;
+		SignatureIndex = RequestOffset / PAK_CACHE_GRANULARITY;
 	}
 
 	check(Data);
@@ -2804,7 +2776,7 @@ void FPakPrecacher::DoSignatureCheck(bool bWasCanceled, IAsyncReadRequest* Reque
 		RequestSize -= PAK_CACHE_GRANULARITY;
 	}
 
-	NewRequestsToLowerComplete(bWasCanceled, Request);
+	NewRequestsToLowerComplete(bWasCanceled, Request, Index);
 }
 
 class FPakAsyncReadFileHandle final : public IAsyncReadFileHandle
diff --git a/Engine/Source/Runtime/Projects/Private/PluginManager.cpp b/Engine/Source/Runtime/Projects/Private/PluginManager.cpp
--- a/Engine/Source/Runtime/Projects/Private/PluginManager.cpp
+++ b/Engine/Source/Runtime/Projects/Private/PluginManager.cpp
@@ -523,17 +523,29 @@ bool FPluginManager::ConfigureEnabledPlugins()
 					}
 				}
 
-				// Load Default<PluginName>.ini config file if it exists
+				// Load <PluginName>.ini config file if it exists
 				FString PluginConfigDir = FPaths::GetPath(Plugin->FileName) / TEXT("Config/");
-				FConfigFile PluginConfig;
-				FConfigCacheIni::LoadExternalIniFile(PluginConfig, *Plugin->Name, *FPaths::EngineConfigDir(), *PluginConfigDir, true);
-				if (PluginConfig.Num() > 0)
+				FString EngineConfigDir = FPaths::EngineConfigDir();
+				FString SourceConfigDir = FPaths::SourceConfigDir();
+
+				// Load Engine plugins out of BasePluginName.ini and the engine directory, game plugins out of DefaultPluginName.ini
+				if (Plugin->LoadedFrom == EPluginLoadedFrom::Engine)
+				{
+					EngineConfigDir = PluginConfigDir;
+				}
+				else
+				{
+					SourceConfigDir = PluginConfigDir;
+				}
+
+				FString PluginConfigFilename = FString::Printf(TEXT("%s%s/%s.ini"), *FPaths::GeneratedConfigDir(), ANSI_TO_TCHAR(FPlatformProperties::PlatformName()), *Plugin->Name);
+				FConfigFile& PluginConfig = GConfig->Add(PluginConfigFilename, FConfigFile());
+
+				// This will write out an ini to PluginConfigFilename
+				if (!FConfigCacheIni::LoadExternalIniFile(PluginConfig, *Plugin->Name, *EngineConfigDir, *SourceConfigDir, true, nullptr, false, true))
 				{
-					FString PlaformName = FPlatformProperties::PlatformName();
-					FString PluginConfigFilename = FString::Printf(TEXT("%s%s/%s.ini"), *FPaths::GeneratedConfigDir(), *PlaformName, *Plugin->Name);
-					FConfigFile& NewConfigFile = GConfig->Add(PluginConfigFilename, FConfigFile());
-					NewConfigFile.AddMissingProperties(PluginConfig);
-					NewConfigFile.Write(PluginConfigFilename);
+					// Nothing to add, remove from map
+					GConfig->Remove(PluginConfigFilename);
 				}
 			}
 		}
diff --git a/Engine/Source/Runtime/SlateCore/Private/Widgets/SWindow.cpp b/Engine/Source/Runtime/SlateCore/Private/Widgets/SWindow.cpp
--- a/Engine/Source/Runtime/SlateCore/Private/Widgets/SWindow.cpp
+++ b/Engine/Source/Runtime/SlateCore/Private/Widgets/SWindow.cpp
@@ -1345,7 +1345,7 @@ EWindowTransparency SWindow::GetTransparencySupport() const
 /** @return A String representation of the widget */
 FString SWindow::ToString() const
 {
-	return FString::Printf( *NSLOCTEXT("SWindow", "Window_Title", " Window : %s ").ToString(), *GetTitle().ToString() );
+	return FText::Format(NSLOCTEXT("SWindow", "Window_TitleFmt", " Window : {0} "), GetTitle()).ToString();
 }
 
 /** @return true if the window should be activated when first shown */
