UE4 - Fixed a critical crash bug with non-EDL loading from pak files.

diff --git a/Engine/Source/Runtime/Core/Public/Async/AsyncFileHandle.h b/Engine/Source/Runtime/Core/Public/Async/AsyncFileHandle.h
--- a/Engine/Source/Runtime/Core/Public/Async/AsyncFileHandle.h
+++ b/Engine/Source/Runtime/Core/Public/Async/AsyncFileHandle.h
@@ -145,7 +145,7 @@ protected:
 	/** Cancel the request. This is a non-blocking async call and so does not ensure completion! **/
 	virtual void CancelImpl() = 0;
 
-	void SetComplete()
+	void SetDataComplete()
 	{
 		bDataIsReady = true;
 		FPlatformMisc::MemoryBarrier();
@@ -154,9 +154,19 @@ protected:
 			Callback(bCanceled, this);
 		}
 		FPlatformMisc::MemoryBarrier();
+	}
+
+	void SetAllComplete()
+	{
 		bCompleteAndCallbackCalled = true;
 		FPlatformMisc::MemoryBarrier();
 	}
+
+	void SetComplete()
+	{
+		SetDataComplete();
+		SetAllComplete();
+	}
 };
 
 class CORE_API IAsyncReadFileHandle
diff --git a/Engine/Source/Runtime/PakFile/Private/IPlatformFilePak.cpp b/Engine/Source/Runtime/PakFile/Private/IPlatformFilePak.cpp
--- a/Engine/Source/Runtime/PakFile/Private/IPlatformFilePak.cpp
+++ b/Engine/Source/Runtime/PakFile/Private/IPlatformFilePak.cpp
@@ -2455,6 +2455,7 @@ public:
 				{
 					WaitEvent->Wait(TimeLimitSeconds * 1000.0f);
 				}
+				FScopeLock Lock(&FPakReadRequestEvent);
 				FPlatformProcess::ReturnSynchEventToPool(WaitEvent);
 				WaitEvent = nullptr;
 			}
@@ -2470,7 +2471,6 @@ public:
 			bRequestOutstanding = false;
 			SetComplete();
 		}
-		check(!WaitEvent); // you canceled from a different thread that you waited from
 	}
 
 	// IPakRequestor Interface
@@ -2518,13 +2518,16 @@ public:
 				check(bCanceled);
 			}
 		}
-		SetComplete();
-		FScopeLock Lock(&FPakReadRequestEvent);
-		bRequestOutstanding = false;
-		if (WaitEvent)
+		SetDataComplete();
 		{
-			WaitEvent->Trigger();
+			FScopeLock Lock(&FPakReadRequestEvent);
+			bRequestOutstanding = false;
+			if (WaitEvent)
+			{
+				WaitEvent->Trigger();
+			}
 		}
+		SetAllComplete();
 	}
 };
 
@@ -2593,14 +2596,16 @@ public:
 				DecryptData(Memory, Align(OriginalSize, FAES::AESBlockSize));
 			}
 		}
-
-		SetComplete();
-		FScopeLock Lock(&FPakReadRequestEvent);
-		bRequestOutstanding = false;
-		if (WaitEvent)
+		SetDataComplete();
 		{
-			WaitEvent->Trigger();
+			FScopeLock Lock(&FPakReadRequestEvent);
+			bRequestOutstanding = false;
+			if (WaitEvent)
+			{
+				WaitEvent->Trigger();
+			}
 		}
+		SetAllComplete();
 	}
 };
 
@@ -2712,6 +2717,7 @@ public:
 				{
 					WaitEvent->Wait(TimeLimitSeconds * 1000.0f);
 				}
+				FScopeLock Lock(&FPakReadRequestEvent);
 				FPlatformProcess::ReturnSynchEventToPool(WaitEvent);
 				WaitEvent = nullptr;
 			}
@@ -2726,7 +2732,6 @@ public:
 			bRequestOutstanding = false;
 			SetComplete();
 		}
-		check(!WaitEvent); // you canceled from a different thread that you waited from
 	}
 
 	void RequestIsComplete()
