Fixed another critical race which potentially could cause a crash in the pak precacher.

diff --git a/Engine/Source/Runtime/PakFile/Private/IPlatformFilePak.cpp b/Engine/Source/Runtime/PakFile/Private/IPlatformFilePak.cpp
--- a/Engine/Source/Runtime/PakFile/Private/IPlatformFilePak.cpp
+++ b/Engine/Source/Runtime/PakFile/Private/IPlatformFilePak.cpp
@@ -2741,13 +2741,16 @@ public:
 		{
 			GatherResults();
 		}
-		SetComplete();
-		FScopeLock Lock(&FPakReadRequestEvent);
-		bRequestOutstanding = false;
-		if (WaitEvent)
+		SetDataComplete();
 		{
-			WaitEvent->Trigger();
+			FScopeLock Lock(&FPakReadRequestEvent);
+			bRequestOutstanding = false;
+			if (WaitEvent)
+			{
+				WaitEvent->Trigger();
+			}
 		}
+		SetAllComplete();
 	}
 	void GatherResults();
 	void DoneWithRawRequests();
