UE4 - Critical fix for a rare race condition in the pak file async IO layer.

diff --git a/Engine/Source/Runtime/PakFile/Private/IPlatformFilePak.cpp b/Engine/Source/Runtime/PakFile/Private/IPlatformFilePak.cpp
--- a/Engine/Source/Runtime/PakFile/Private/IPlatformFilePak.cpp
+++ b/Engine/Source/Runtime/PakFile/Private/IPlatformFilePak.cpp
@@ -2483,31 +2483,28 @@ public:
 
 	virtual void WaitCompletionImpl(float TimeLimitSeconds) override
 	{
-		if (bRequestOutstanding)
 		{
+			FScopeLock Lock(&FPakReadRequestEvent);
+			if (bRequestOutstanding)
 			{
-				FScopeLock Lock(&FPakReadRequestEvent);
-				if (bRequestOutstanding)
-				{
-					check(!WaitEvent);
-					WaitEvent = FPlatformProcess::GetSynchEventFromPool(true);
-				}
+				check(!WaitEvent);
+				WaitEvent = FPlatformProcess::GetSynchEventFromPool(true);
 			}
-			if (WaitEvent)
+		}
+		if (WaitEvent)
+		{
+			if (TimeLimitSeconds == 0.0f)
 			{
-				if (TimeLimitSeconds == 0.0f)
-				{
-					WaitEvent->Wait();
-					check(!bRequestOutstanding);
-				}
-				else
-				{
-					WaitEvent->Wait(TimeLimitSeconds * 1000.0f);
-				}
-				FScopeLock Lock(&FPakReadRequestEvent);
-				FPlatformProcess::ReturnSynchEventToPool(WaitEvent);
-				WaitEvent = nullptr;
+				WaitEvent->Wait();
+				check(!bRequestOutstanding);
 			}
+			else
+			{
+				WaitEvent->Wait(TimeLimitSeconds * 1000.0f);
+			}
+			FScopeLock Lock(&FPakReadRequestEvent);
+			FPlatformProcess::ReturnSynchEventToPool(WaitEvent);
+			WaitEvent = nullptr;
 		}
 	}
 	virtual void CancelImpl() override
@@ -2746,31 +2743,28 @@ public:
 
 	virtual void WaitCompletionImpl(float TimeLimitSeconds) override
 	{
-		if (bRequestOutstanding)
 		{
+			FScopeLock Lock(&FPakReadRequestEvent);
+			if (bRequestOutstanding)
 			{
-				FScopeLock Lock(&FPakReadRequestEvent);
-				if (bRequestOutstanding)
-				{
-					check(!WaitEvent);
-					WaitEvent = FPlatformProcess::GetSynchEventFromPool(true);
-				}
+				check(!WaitEvent);
+				WaitEvent = FPlatformProcess::GetSynchEventFromPool(true);
 			}
-			if (WaitEvent)
+		}
+		if (WaitEvent)
+		{
+			if (TimeLimitSeconds == 0.0f)
 			{
-				if (TimeLimitSeconds == 0.0f)
-				{
-					WaitEvent->Wait();
-					check(!bRequestOutstanding);
-				}
-				else
-				{
-					WaitEvent->Wait(TimeLimitSeconds * 1000.0f);
-				}
-				FScopeLock Lock(&FPakReadRequestEvent);
-				FPlatformProcess::ReturnSynchEventToPool(WaitEvent);
-				WaitEvent = nullptr;
+				WaitEvent->Wait();
+				check(!bRequestOutstanding);
 			}
+			else
+			{
+				WaitEvent->Wait(TimeLimitSeconds * 1000.0f);
+			}
+			FScopeLock Lock(&FPakReadRequestEvent);
+			FPlatformProcess::ReturnSynchEventToPool(WaitEvent);
+			WaitEvent = nullptr;
 		}
 	}
 	virtual void CancelImpl() override
