PR #4452: Fixed FindFilesRecursively in IPlatformFilePak (Contributed by [REDACTED_PERSON])

diff --git a/Engine/Source/Runtime/PakFile/Public/IPlatformFilePak.h b/Engine/Source/Runtime/PakFile/Public/IPlatformFilePak.h
--- a/Engine/Source/Runtime/PakFile/Public/IPlatformFilePak.h
+++ b/Engine/Source/Runtime/PakFile/Public/IPlatformFilePak.h
@@ -1626,7 +1626,7 @@ public:
 	{
 		if (LowerLevel->DirectoryExists(Directory))
 		{
-			LowerLevel->FindFiles(FoundFiles, Directory, FileExtension);
+			LowerLevel->FindFilesRecursively(FoundFiles, Directory, FileExtension);
 		}
 		
 		bool bRecursive = true;
