Add support for checking if persistent storage is available before attempting to use it.

This fixes a regression on some platforms where persistant storage is optional and would warn if it is not available.

diff --git a/Engine/Source/Runtime/Core/Private/Misc/Paths.cpp b/Engine/Source/Runtime/Core/Private/Misc/Paths.cpp
--- a/Engine/Source/Runtime/Core/Private/Misc/Paths.cpp
+++ b/Engine/Source/Runtime/Core/Private/Misc/Paths.cpp
@@ -303,6 +303,11 @@ FString FPaths::ProjectModsDir()
 	return FPaths::ProjectDir() + TEXT("Mods/");
 }
 
+bool FPaths::HasProjectPersistentDownloadDir()
+{
+	return FPlatformMisc::HasProjectPersistentDownloadDir();
+}
+
 FString FPaths::ProjectPersistentDownloadDir()
 {
 	return FPlatformMisc::GamePersistentDownloadDir();
diff --git a/Engine/Source/Runtime/Core/Public/GenericPlatform/GenericPlatformMisc.h b/Engine/Source/Runtime/Core/Public/GenericPlatform/GenericPlatformMisc.h
--- a/Engine/Source/Runtime/Core/Public/GenericPlatform/GenericPlatformMisc.h
+++ b/Engine/Source/Runtime/Core/Public/GenericPlatform/GenericPlatformMisc.h
@@ -762,6 +762,16 @@ public:
 	*/
 	static FString CloudDir();
 
+	/**
+	*	Return true if the PersistentDownloadDir is available.
+	*	On some platforms, a writable directory might not be available by default.
+	*	Using this function allows handling that case early.
+	*/
+	static bool HasProjectPersistentDownloadDir()
+	{
+		return true;
+	}
+
 	/**
 	*	Return the GamePersistentDownloadDir.  
 	*	On some platforms, returns the writable directory for downloaded data that persists across play sessions.
diff --git a/Engine/Source/Runtime/Core/Public/Misc/Paths.h b/Engine/Source/Runtime/Core/Public/Misc/Paths.h
--- a/Engine/Source/Runtime/Core/Public/Misc/Paths.h
+++ b/Engine/Source/Runtime/Core/Public/Misc/Paths.h
@@ -192,6 +192,11 @@ public:
 	 */
 	static FString ProjectModsDir();
 
+	/*
+	* Returns true if a writable directory for downloaded data that persists across play sessions is available
+	*/
+	static bool HasProjectPersistentDownloadDir();
+
 	/*
 	* Returns the writable directory for downloaded data that persists across play sessions.
 	*/
diff --git a/Engine/Source/Runtime/Launch/Private/LaunchEngineLoop.cpp b/Engine/Source/Runtime/Launch/Private/LaunchEngineLoop.cpp
--- a/Engine/Source/Runtime/Launch/Private/LaunchEngineLoop.cpp
+++ b/Engine/Source/Runtime/Launch/Private/LaunchEngineLoop.cpp
@@ -1998,8 +1998,11 @@ int32 FEngineLoop::PreInit(const TCHAR* CmdLine)
 		if (FPlatformProperties::RequiresCookedData())
 		{
 			// Open the game library which contains the material shaders.
-	        FShaderCodeLibrary::OpenLibrary(FApp::GetProjectName(), FPaths::ProjectContentDir());
-			FShaderCodeLibrary::OpenLibrary(FApp::GetProjectName(), FPaths::Combine(*FPaths::ProjectPersistentDownloadDir(), TEXT("InstalledContent"), FApp::GetProjectName(), TEXT("Content")));
+			FShaderCodeLibrary::OpenLibrary(FApp::GetProjectName(), FPaths::ProjectContentDir());
+			if (FPaths::HasProjectPersistentDownloadDir())
+			{
+				FShaderCodeLibrary::OpenLibrary(FApp::GetProjectName(), FPaths::Combine(*FPaths::ProjectPersistentDownloadDir(), TEXT("InstalledContent"), FApp::GetProjectName(), TEXT("Content")));
+			}
 		}
 		
 		InitGameTextLocalization();
diff --git a/Engine/Source/Runtime/PakFile/Public/IPlatformFilePak.h b/Engine/Source/Runtime/PakFile/Public/IPlatformFilePak.h
--- a/Engine/Source/Runtime/PakFile/Public/IPlatformFilePak.h
+++ b/Engine/Source/Runtime/PakFile/Public/IPlatformFilePak.h
@@ -1681,9 +1681,10 @@ public:
 		FString StandardDirectory = Directory;
 		FPaths::MakeStandardFilename(StandardDirectory);
 		
+		bool bIsDownloadableDir = (FPaths::HasProjectPersistentDownloadDir() && StandardDirectory.StartsWith(FPaths::ProjectPersistentDownloadDir())) || StandardDirectory.StartsWith(FPaths::CloudDir());
+
 		// don't look for in pak files for target-only locations
-		if (!StandardDirectory.StartsWith(FPaths::ProjectPersistentDownloadDir()) &&
-			!StandardDirectory.StartsWith(FPaths::CloudDir()))
+		if (!bIsDownloadableDir)
 		{
 			GetMountedPaks(Paks);
 		}
@@ -1792,9 +1793,10 @@ public:
 		FString StandardDirectory = Directory;
 		FPaths::MakeStandardFilename(StandardDirectory);
 
+		bool bIsDownloadableDir = (FPaths::HasProjectPersistentDownloadDir() && StandardDirectory.StartsWith(FPaths::ProjectPersistentDownloadDir())) || StandardDirectory.StartsWith(FPaths::CloudDir());
+
 		// don't look for in pak files for target-only locations
-		if (!StandardDirectory.StartsWith(FPaths::ProjectPersistentDownloadDir()) &&
-			!StandardDirectory.StartsWith(FPaths::CloudDir()))
+		if (!bIsDownloadableDir)
 		{
 			GetMountedPaks(Paks);
 		}
