Optimize deserialization of compressed pak files.
Reuse previously decompressed buffer where possible to avoid repeatedly deserialising and decompressing the same block when serializing data < compressed block size.

diff --git a/Engine/Source/Runtime/PakFile/Private/IPlatformFilePak.cpp b/Engine/Source/Runtime/PakFile/Private/IPlatformFilePak.cpp
--- a/Engine/Source/Runtime/PakFile/Private/IPlatformFilePak.cpp
+++ b/Engine/Source/Runtime/PakFile/Private/IPlatformFilePak.cpp
@@ -1,4 +1,4 @@
-// Copyright 1998-2019 Epic Games, Inc. All Rights Reserved.
+// Copyright 1998-2018 Epic Games, Inc. All Rights Reserved.
 
 #include "IPlatformFilePak.h"
 #include "HAL/FileManager.h"
@@ -3882,6 +3882,8 @@ public:
 	FCompressionScratchBuffers()
 		: TempBufferSize(0)
 		, ScratchBufferSize(0)
+		, LastReader(nullptr)
+		, LastDecompressedBlock(0xFFFFFFFF)
 	{}
 
 	int64				TempBufferSize;
@@ -3889,6 +3891,9 @@ public:
 	int64				ScratchBufferSize;
 	TUniquePtr<uint8[]>	ScratchBuffer;
 
+	void* LastReader;
+	uint32 LastDecompressedBlock;
+
 	void EnsureBufferSpace(int64 CompressionBlockSize, int64 ScrachSize)
 	{
 		if (TempBufferSize < CompressionBlockSize)
@@ -3951,6 +3956,16 @@ public:
 	{
 	}
 
+	~FPakCompressedReaderPolicy()
+	{
+		FCompressionScratchBuffers& ScratchSpace = FCompressionScratchBuffers::Get();
+		if(ScratchSpace.LastReader == this)
+		{
+			ScratchSpace.LastDecompressedBlock = 0xFFFFFFFF;
+			ScratchSpace.LastReader = nullptr;
+		}
+	}
+
 	/** Pak file that own this file data */
 	const FPakFile&		PakFile;
 	/** Pak file entry for this file. */
@@ -3986,6 +4001,7 @@ public:
 		float SlopMultiplier = 1.1f;
 		int64 WorkingBufferRequiredSize = FCompression::CompressMemoryBound(CompressionMethod, CompressionBlockSize) * SlopMultiplier;
 		WorkingBufferRequiredSize = EncryptionPolicy::AlignReadRequest(WorkingBufferRequiredSize);
+		const bool bExistingScratchBufferValid = ScratchSpace.TempBufferSize >= CompressionBlockSize;
 		ScratchSpace.EnsureBufferSpace(CompressionBlockSize, WorkingBufferRequiredSize * 2);
 		WorkingBuffers[0] = ScratchSpace.ScratchBuffer.Get();
 		WorkingBuffers[1] = ScratchSpace.ScratchBuffer.Get() + WorkingBufferRequiredSize;
@@ -4007,6 +4023,21 @@ public:
 
 			int64 ReadSize = EncryptionPolicy::AlignReadRequest(CompressedBlockSize);
 			int64 WriteSize = FMath::Min<int64>(UncompressedBlockSize - DirectCopyStart, Length);
+
+			const bool bCurrentScratchTempBufferValid = 
+				bExistingScratchBufferValid && !bStartedUncompress
+				// ensure this object was the last reader from the scratch buffer and the last thing it decompressed was this block.
+				&& ScratchSpace.LastReader == this && ScratchSpace.LastDecompressedBlock == CompressionBlockIndex 
+				// ensure the previous decompression destination was the scratch buffer.
+				&& !(DirectCopyStart == 0 && Length >= CompressionBlockSize); 
+
+			if (bCurrentScratchTempBufferValid)
+			{
+				// Reuse the existing scratch buffer to avoid repeatedly deserializing and decompressing the same block.
+				FMemory::Memcpy(V, ScratchSpace.TempBuffer.Get() + DirectCopyStart, WriteSize);
+			}
+			else
+			{
 			PakReader->Seek(Block.CompressedStart + (PakFile.GetInfo().HasRelativeCompressedChunkOffsets() ? PakEntry.Offset : 0));
 			PakReader->Serialize(WorkingBuffers[CompressionBlockIndex & 1], ReadSize);
 			if (bStartedUncompress)
@@ -4027,6 +4058,8 @@ public:
 				TaskDetails.CompressedBuffer = WorkingBuffers[CompressionBlockIndex & 1];
 				TaskDetails.CompressedSize = CompressedBlockSize;
 				TaskDetails.CopyOut = nullptr;
+					ScratchSpace.LastDecompressedBlock = 0xFFFFFFFF;
+					ScratchSpace.LastReader = nullptr;
 			}
 			else
 			{
@@ -4039,6 +4072,9 @@ public:
 				TaskDetails.CopyOut = V;
 				TaskDetails.CopyOffset = DirectCopyStart;
 				TaskDetails.CopyLength = WriteSize;
+
+					ScratchSpace.LastDecompressedBlock = CompressionBlockIndex;
+					ScratchSpace.LastReader = this;
 			}
 
 			if (Length == WriteSize)
@@ -4050,6 +4086,8 @@ public:
 				UncompressTask.StartBackgroundTask();
 			}
 			bStartedUncompress = true;
+			}
+		
 			V = (void*)((uint8*)V + WriteSize);
 			Length -= WriteSize;
 			DirectCopyStart = 0;
@@ -4242,38 +4280,38 @@ void FPakFile::Initialize(FArchive* Reader)
 	// start up one to offset the -- below
 	CompatibleVersion++;
 	do
-	{
+		{
 		// try the next version down
-		CompatibleVersion--;
-		// go to start
+			CompatibleVersion--;
+
 		Reader->Seek(CachedTotalSize - Info.GetSerializedSize(CompatibleVersion));
-		
-		// read it in (this will check size, etc, and is considered safe)
+
+		// Serialize trailer and check if everything is as expected.
 		Info.Serialize(*Reader, CompatibleVersion);
 	}
 	while (Info.Magic != FPakInfo::PakFile_Magic && CompatibleVersion >= FPakInfo::PakFile_Version_Initial);
 
-	UE_CLOG(Info.Magic != FPakInfo::PakFile_Magic, LogPakFile, Fatal, TEXT("Trailing magic number (%ud) in '%s' is different than the expected one. Verify your installation."), Info.Magic, *PakFilename);
-	UE_CLOG(!(Info.Version >= FPakInfo::PakFile_Version_Initial && Info.Version <= CompatibleVersion), LogPakFile, Fatal, TEXT("Invalid pak file version (%d) in '%s'. Verify your installation."), Info.Version, *PakFilename);
+		UE_CLOG(Info.Magic != FPakInfo::PakFile_Magic, LogPakFile, Fatal, TEXT("Trailing magic number (%ud) in '%s' is different than the expected one. Verify your installation."), Info.Magic, *PakFilename);
+		UE_CLOG(!(Info.Version >= FPakInfo::PakFile_Version_Initial && Info.Version <= CompatibleVersion), LogPakFile, Fatal, TEXT("Invalid pak file version (%d) in '%s'. Verify your installation."), Info.Version, *PakFilename);
 	UE_CLOG((Info.bEncryptedIndex == 1) && (!FCoreDelegates::GetPakEncryptionKeyDelegate().IsBound()), LogPakFile, Fatal, TEXT("Index of pak file '%s' is encrypted, but this executable doesn't have any valid decryption keys"), *PakFilename);
-	UE_CLOG(!(Info.IndexOffset >= 0 && Info.IndexOffset < CachedTotalSize), LogPakFile, Fatal, TEXT("Index offset for pak file '%s' is invalid (%lld)"), *PakFilename, Info.IndexOffset);
-	UE_CLOG(!((Info.IndexOffset + Info.IndexSize) >= 0 && (Info.IndexOffset + Info.IndexSize) <= CachedTotalSize), LogPakFile, Fatal, TEXT("Index end offset for pak file '%s' is invalid (%lld)"), *PakFilename, Info.IndexOffset + Info.IndexSize);
-
-	// If we aren't using a dynamic encryption key, process the pak file using the embedded key
-	if (!Info.EncryptionKeyGuid.IsValid() || GetRegisteredEncryptionKeys().HasKey(Info.EncryptionKeyGuid))
+		UE_CLOG(!(Info.IndexOffset >= 0 && Info.IndexOffset < CachedTotalSize), LogPakFile, Fatal, TEXT("Index offset for pak file '%s' is invalid (%lld)"), *PakFilename, Info.IndexOffset);
+		UE_CLOG(!((Info.IndexOffset + Info.IndexSize) >= 0 && (Info.IndexOffset + Info.IndexSize) <= CachedTotalSize), LogPakFile, Fatal, TEXT("Index end offset for pak file '%s' is invalid (%lld)"), *PakFilename, Info.IndexOffset + Info.IndexSize);
 
-	{
-		LoadIndex(Reader);
+		// If we aren't using a dynamic encryption key, process the pak file using the embedded key
+		if (!Info.EncryptionKeyGuid.IsValid() || GetRegisteredEncryptionKeys().HasKey(Info.EncryptionKeyGuid))
 
-		if (FParse::Param(FCommandLine::Get(), TEXT("checkpak")))
 		{
-			ensure(Check());
-		}
+			LoadIndex(Reader);
 
-		// LoadIndex should crash in case of an error, so just assume everything is ok if we got here.
-		bIsValid = true;
+			if (FParse::Param(FCommandLine::Get(), TEXT("checkpak")))
+			{
+				ensure(Check());
+			}
+
+			// LoadIndex should crash in case of an error, so just assume everything is ok if we got here.
+			bIsValid = true;
+		}
 	}
-}
 
 void FPakFile::LoadIndex(FArchive* Reader)
 {
