Fix comments somehow lost during merge

diff --git a/Engine/Source/Runtime/PakFile/Private/IPlatformFilePak.cpp b/Engine/Source/Runtime/PakFile/Private/IPlatformFilePak.cpp
--- a/Engine/Source/Runtime/PakFile/Private/IPlatformFilePak.cpp
+++ b/Engine/Source/Runtime/PakFile/Private/IPlatformFilePak.cpp
@@ -1,4 +1,4 @@
-// Copyright 1998-2018 Epic Games, Inc. All Rights Reserved.
+// Copyright 1998-2019 Epic Games, Inc. All Rights Reserved.
 
 #include "IPlatformFilePak.h"
 #include "HAL/FileManager.h"
@@ -4283,10 +4283,10 @@ void FPakFile::Initialize(FArchive* Reader)
 		{
 		// try the next version down
 			CompatibleVersion--;
-
+		// go to start
 		Reader->Seek(CachedTotalSize - Info.GetSerializedSize(CompatibleVersion));
 
-		// Serialize trailer and check if everything is as expected.
+		// read it in (this will check size, etc, and is considered safe)
 		Info.Serialize(*Reader, CompatibleVersion);
 	}
 	while (Info.Magic != FPakInfo::PakFile_Magic && CompatibleVersion >= FPakInfo::PakFile_Version_Initial);
