Fixed the issue where unrealpak failed to read from old pak files with no content.

diff --git a/Engine/Source/Runtime/PakFile/Private/IPlatformFilePak.cpp b/Engine/Source/Runtime/PakFile/Private/IPlatformFilePak.cpp
--- a/Engine/Source/Runtime/PakFile/Private/IPlatformFilePak.cpp
+++ b/Engine/Source/Runtime/PakFile/Private/IPlatformFilePak.cpp
@@ -4266,6 +4266,7 @@ FArchive* FPakFile::SetupSignedPakReader(FArchive* ReaderArchive, const TCHAR* F
 void FPakFile::Initialize(FArchive* Reader)
 {
 	CachedTotalSize = Reader->TotalSize();
+	bool bShouldLoad = false;
 	int32 CompatibleVersion = FPakInfo::PakFile_Version_Latest;
 
 	LLM_SCOPE(ELLMTag::FileSystem);
@@ -4273,27 +4274,36 @@ void FPakFile::Initialize(FArchive* Reader)
 	// Serialize trailer and check if everything is as expected.
 	// start up one to offset the -- below
 	CompatibleVersion++;
+	int64 FileInfoPos = -1;
 	do
-		{
+	{
 		// try the next version down
-			CompatibleVersion--;
-		// go to start
-		Reader->Seek(CachedTotalSize - Info.GetSerializedSize(CompatibleVersion));
+		CompatibleVersion--;
 
-		// read it in (this will check size, etc, and is considered safe)
-		Info.Serialize(*Reader, CompatibleVersion);
-	}
-	while (Info.Magic != FPakInfo::PakFile_Magic && CompatibleVersion >= FPakInfo::PakFile_Version_Initial);
+		FileInfoPos = CachedTotalSize - Info.GetSerializedSize(CompatibleVersion);
+		if (FileInfoPos >= 0)
+		{
+			Reader->Seek(FileInfoPos);
 
+			// Serialize trailer and check if everything is as expected.
+			Info.Serialize(*Reader, CompatibleVersion);
+			if (Info.Magic == FPakInfo::PakFile_Magic)
+			{
+				bShouldLoad = true;
+			}
+		}
+	} while (!bShouldLoad && CompatibleVersion >= FPakInfo::PakFile_Version_Initial);
+
+	if (bShouldLoad)
+	{
 		UE_CLOG(Info.Magic != FPakInfo::PakFile_Magic, LogPakFile, Fatal, TEXT("Trailing magic number (%ud) in '%s' is different than the expected one. Verify your installation."), Info.Magic, *PakFilename);
 		UE_CLOG(!(Info.Version >= FPakInfo::PakFile_Version_Initial && Info.Version <= CompatibleVersion), LogPakFile, Fatal, TEXT("Invalid pak file version (%d) in '%s'. Verify your installation."), Info.Version, *PakFilename);
-	UE_CLOG((Info.bEncryptedIndex == 1) && (!FCoreDelegates::GetPakEncryptionKeyDelegate().IsBound()), LogPakFile, Fatal, TEXT("Index of pak file '%s' is encrypted, but this executable doesn't have any valid decryption keys"), *PakFilename);
+		UE_CLOG((Info.bEncryptedIndex == 1) && (!FCoreDelegates::GetPakEncryptionKeyDelegate().IsBound()), LogPakFile, Fatal, TEXT("Index of pak file '%s' is encrypted, but this executable doesn't have any valid decryption keys"), *PakFilename);
 		UE_CLOG(!(Info.IndexOffset >= 0 && Info.IndexOffset < CachedTotalSize), LogPakFile, Fatal, TEXT("Index offset for pak file '%s' is invalid (%lld)"), *PakFilename, Info.IndexOffset);
 		UE_CLOG(!((Info.IndexOffset + Info.IndexSize) >= 0 && (Info.IndexOffset + Info.IndexSize) <= CachedTotalSize), LogPakFile, Fatal, TEXT("Index end offset for pak file '%s' is invalid (%lld)"), *PakFilename, Info.IndexOffset + Info.IndexSize);
 
 		// If we aren't using a dynamic encryption key, process the pak file using the embedded key
 		if (!Info.EncryptionKeyGuid.IsValid() || GetRegisteredEncryptionKeys().HasKey(Info.EncryptionKeyGuid))
-
 		{
 			LoadIndex(Reader);
 
@@ -4306,6 +4316,7 @@ void FPakFile::Initialize(FArchive* Reader)
 			bIsValid = true;
 		}
 	}
+}
 
 void FPakFile::LoadIndex(FArchive* Reader)
 {
