Accidentally forced pak filename/entry compression on. Disabling again.

diff --git a/Engine/Source/Runtime/PakFile/Private/IPlatformFilePak.cpp b/Engine/Source/Runtime/PakFile/Private/IPlatformFilePak.cpp
--- a/Engine/Source/Runtime/PakFile/Private/IPlatformFilePak.cpp
+++ b/Engine/Source/Runtime/PakFile/Private/IPlatformFilePak.cpp
@@ -5419,8 +5419,8 @@ bool FPakPlatformFile::Initialize(IPlatformFile* Inner, const TCHAR* CmdLine)
 #if !(IS_PROGRAM || WITH_EDITOR)
 	FCoreDelegates::OnFEngineLoopInitComplete.AddLambda([this] {
 			FPlatformMisc::LowLevelOutputDebugStringf(TEXT("Checking Pak Config"));
-			bool bUnloadPakEntryFilenamesIfPossible = true;
-			//GConfig->GetBool(TEXT("Pak"), TEXT("UnloadPakEntryFilenamesIfPossible"), bUnloadPakEntryFilenamesIfPossible, GEngineIni);
+			bool bUnloadPakEntryFilenamesIfPossible = false;
+			GConfig->GetBool(TEXT("Pak"), TEXT("UnloadPakEntryFilenamesIfPossible"), bUnloadPakEntryFilenamesIfPossible, GEngineIni);
 
 			if (bUnloadPakEntryFilenamesIfPossible)
 			{
@@ -5439,8 +5439,8 @@ bool FPakPlatformFile::Initialize(IPlatformFile* Inner, const TCHAR* CmdLine)
 				PakPlatformFile->UnloadPakEntryFilenames(&DirectoryRootsToKeep);
 			}
 
-			bool bShrinkPakEntriesMemoryUsage = true;
-			//GConfig->GetBool(TEXT("Pak"), TEXT("ShrinkPakEntriesMemoryUsage"), bShrinkPakEntriesMemoryUsage, GEngineIni);
+			bool bShrinkPakEntriesMemoryUsage = false;
+			GConfig->GetBool(TEXT("Pak"), TEXT("ShrinkPakEntriesMemoryUsage"), bShrinkPakEntriesMemoryUsage, GEngineIni);
 			if (bShrinkPakEntriesMemoryUsage)
 			{
 				FPakPlatformFile* PakPlatformFile = (FPakPlatformFile*)(FPlatformFileManager::Get().FindPlatformFile(FPakPlatformFile::GetTypeName()));
