Initialize signature checking hashcache array when using CRCs. Static analysis warning fix.

diff --git a/Engine/Source/Runtime/PakFile/Private/IPlatformFilePak.cpp b/Engine/Source/Runtime/PakFile/Private/IPlatformFilePak.cpp
--- a/Engine/Source/Runtime/PakFile/Private/IPlatformFilePak.cpp
+++ b/Engine/Source/Runtime/PakFile/Private/IPlatformFilePak.cpp
@@ -3114,7 +3114,12 @@ void FPakPrecacher::DoSignatureCheck(bool bWasCanceled, IAsyncReadRequest* Reque
 	uint16 PakIndex;
 	FSHAHash MasterSignatureHash;
 	static const int64 MaxHashesToCache = 16;
+
+#if PAKHASH_USE_CRC
+	TPakChunkHash HashCache[MaxHashesToCache] = { 0 };
+#else
 	TPakChunkHash HashCache[MaxHashesToCache];
+#endif
 
 	{
 		// Try and keep lock for as short a time as possible. Find our request and copy out the data we need
