[LLM] add missing FileSystem tag to file request. moves about ~2MB on average out of Untagged.

diff --git a/Engine/Source/Runtime/PakFile/Private/IPlatformFilePak.cpp b/Engine/Source/Runtime/PakFile/Private/IPlatformFilePak.cpp
--- a/Engine/Source/Runtime/PakFile/Private/IPlatformFilePak.cpp
+++ b/Engine/Source/Runtime/PakFile/Private/IPlatformFilePak.cpp
@@ -3289,6 +3289,8 @@ public:
 	}
 	virtual IAsyncReadRequest* ReadRequest(int64 Offset, int64 BytesToRead, EAsyncIOPriorityAndFlags PriorityAndFlags = AIOP_Normal, FAsyncFileCallBack* CompleteCallback = nullptr, uint8* UserSuppliedMemory = nullptr) override
 	{
+		LLM_SCOPE(ELLMTag::FileSystem);
+
 		if (BytesToRead == MAX_int64)
 		{
 			BytesToRead = UncompressedFileSize - Offset;
