- Replaced -nommio with a CVar: mmio.enable. Can use commandline to disable with: -dpcvars=mmio.enable=0
- This allows it to be hotfixed off

diff --git a/Engine/Source/Runtime/PakFile/Private/IPlatformFilePak.cpp b/Engine/Source/Runtime/PakFile/Private/IPlatformFilePak.cpp
--- a/Engine/Source/Runtime/PakFile/Private/IPlatformFilePak.cpp
+++ b/Engine/Source/Runtime/PakFile/Private/IPlatformFilePak.cpp
@@ -3811,16 +3811,20 @@ static FAutoConsoleCommand MappedFileTestCmd(
 );
 #endif
 
+static int32 GMMIO_Enable = 1;
+static FAutoConsoleVariableRef CVar_MMIOEnable(
+	   TEXT("mmio.enable"),
+	   GMMIO_Enable,
+	   TEXT("If > 0, then enable memory mapped IO on platforms that support it.")
+	   );
+
+
 IMappedFileHandle* FPakPlatformFile::OpenMapped(const TCHAR* Filename)
 {
-#if !UE_BUILD_SHIPPING
-	// disable all mmio if commandline requested it
-	static bool bNoMMIO = FParse::Param(FCommandLine::Get(), TEXT("nommio"));
-	if (bNoMMIO)
+	if (!GMMIO_Enable)
 	{
 		return nullptr;
 	}
-#endif
 
 	// Check pak files first
 	FPakEntry FileEntry;
