Delete record logging changed to Verbose.

diff --git a/Engine/Source/Runtime/PakFile/Public/IPlatformFilePak.h b/Engine/Source/Runtime/PakFile/Public/IPlatformFilePak.h
--- a/Engine/Source/Runtime/PakFile/Public/IPlatformFilePak.h
+++ b/Engine/Source/Runtime/PakFile/Public/IPlatformFilePak.h
@@ -1604,7 +1604,7 @@ public:
 			if (DeletedReadOrder != -1 && DeletedReadOrder > PakReadOrder)
 			{
 				//found a delete record in a higher priority patch level, but now we're at a lower priority set - don't search further back or we'll find the original, old file.
-				UE_LOG( LogPakFile, Display, TEXT("Delete Record: Accepted a delete record for %s"), Filename );
+				UE_LOG( LogPakFile, Verbose, TEXT("Delete Record: Accepted a delete record for %s"), Filename );
 				return false;
 			}
 
@@ -1615,13 +1615,13 @@ public:
 				{
 					*OutPakFile = Paks[PakIndex].PakFile;
 				}
-				UE_CLOG( DeletedReadOrder != -1, LogPakFile, Display, TEXT("Delete Record: Ignored delete record for %s - found it in %s instead (asset was moved between chunks)"), Filename, *Paks[PakIndex].PakFile->GetFilename() );
+				UE_CLOG( DeletedReadOrder != -1, LogPakFile, Verbose, TEXT("Delete Record: Ignored delete record for %s - found it in %s instead (asset was moved between chunks)"), Filename, *Paks[PakIndex].PakFile->GetFilename() );
 				return true;
 			}
 			else if (FindResult == FPakFile::EFindResult::FoundDeleted )
 			{
 				DeletedReadOrder = PakReadOrder;
-				UE_LOG( LogPakFile, Display, TEXT("Delete Record: Found a delete record for %s in %s"), Filename, *Paks[PakIndex].PakFile->GetFilename() );
+				UE_LOG( LogPakFile, Verbose, TEXT("Delete Record: Found a delete record for %s in %s"), Filename, *Paks[PakIndex].PakFile->GetFilename() );
 			}
 		}
 
