allow turning pak unload off from commandline with -nounloadpakentries, this allows all items to be properly loaded in PSO gen

diff --git a/Engine/Source/Runtime/PakFile/Private/IPlatformFilePak.cpp b/Engine/Source/Runtime/PakFile/Private/IPlatformFilePak.cpp
--- a/Engine/Source/Runtime/PakFile/Private/IPlatformFilePak.cpp
+++ b/Engine/Source/Runtime/PakFile/Private/IPlatformFilePak.cpp
@@ -5454,7 +5454,7 @@ bool FPakPlatformFile::Initialize(IPlatformFile* Inner, const TCHAR* CmdLine)
 			bool bUnloadPakEntryFilenamesIfPossible = FParse::Param(FCommandLine::Get(), TEXT("unloadpakentryfilenames"));
 			GConfig->GetBool(TEXT("Pak"), TEXT("UnloadPakEntryFilenamesIfPossible"), bUnloadPakEntryFilenamesIfPossible, GEngineIni);
 
-            if (bUnloadPakEntryFilenamesIfPossible || FParse::Param(FCommandLine::Get(),TEXT("unloadpakentries")))
+        if ((bUnloadPakEntryFilenamesIfPossible && !FParse::Param(FCommandLine::Get(), TEXT("nounloadpakentries"))) || FParse::Param(FCommandLine::Get(),TEXT("unloadpakentries")))
 			{
 				// With [Pak] UnloadPakEntryFilenamesIfPossible enabled, [Pak] DirectoryRootsToKeepInMemoryWhenUnloadingPakEntryFilenames
 				// can contain pak entry directory wildcards of which the entire recursive directory structure of filenames underneath a
