Fix debug logging of chunk hashes for pak file signature check failures

diff --git a/Engine/Source/Runtime/PakFile/Private/IPlatformFilePak.cpp b/Engine/Source/Runtime/PakFile/Private/IPlatformFilePak.cpp
--- a/Engine/Source/Runtime/PakFile/Private/IPlatformFilePak.cpp
+++ b/Engine/Source/Runtime/PakFile/Private/IPlatformFilePak.cpp
@@ -158,6 +158,15 @@ TPakChunkHash ComputePakChunkHash(const void* InData, int64 InDataSizeInBytes)
 #endif
 }
 
+FString ChunkHashToString(const TPakChunkHash& InHash)
+{
+#if PAKHASH_USE_CRC
+	return FString::Printf(TEXT("%08X"), InHash);
+#else
+	return LexToString(InHash);
+#endif
+}
+
 #ifndef EXCLUDE_NONPAK_UE_EXTENSIONS
 #define EXCLUDE_NONPAK_UE_EXTENSIONS 1	// Use .Build.cs file to disable this if the game relies on accessing loose files
 #endif
@@ -3180,7 +3189,7 @@ void FPakPrecacher::DoSignatureCheck(bool bWasCanceled, IAsyncReadRequest* Reque
 				FScopeLock Lock(&CachedFilesScopeLock);
 				FPakData* PakData = &CachedPakData[PakIndex];
 
-				UE_LOG(LogPakFile, Warning, TEXT("Pak chunk signing mismatch on chunk [%i/%i]! Expected 0x%8X, Received 0x%8X"), SignatureIndex, PakData->Signatures.ChunkHashes.Num(), *LexToString(PakData->Signatures.ChunkHashes[SignatureIndex]), *LexToString(ThisHash));
+				UE_LOG(LogPakFile, Warning, TEXT("Pak chunk signing mismatch on chunk [%i/%i]! Expected %s, Received %s"), SignatureIndex, PakData->Signatures.ChunkHashes.Num(), *ChunkHashToString(PakData->Signatures.ChunkHashes[SignatureIndex]), *ChunkHashToString(ThisHash));
 
 				// Check the signatures are still as we expected them
 				if (PakData->Signatures.DecryptedHash != PakData->Signatures.ComputeCurrentMasterHash())
diff --git a/Engine/Source/Runtime/PakFile/Private/SignedArchiveReader.cpp b/Engine/Source/Runtime/PakFile/Private/SignedArchiveReader.cpp
--- a/Engine/Source/Runtime/PakFile/Private/SignedArchiveReader.cpp
+++ b/Engine/Source/Runtime/PakFile/Private/SignedArchiveReader.cpp
@@ -244,7 +244,7 @@ bool FChunkCacheWorker::CheckSignature(const FChunkRequest& ChunkInfo)
 	bool bChunkHashesMatch = (ChunkHash == Signatures.ChunkHashes[ChunkInfo.Index]);
 	if (!bChunkHashesMatch)
 	{
-		UE_LOG(LogPakFile, Warning, TEXT("Pak chunk signing mismatch on chunk [%i/%i]! Expected 0x%8X, Received 0x%8X"), ChunkInfo.Index, Signatures.ChunkHashes.Num(), *LexToString(Signatures.ChunkHashes[ChunkInfo.Index]), *LexToString(ChunkHash));
+		UE_LOG(LogPakFile, Warning, TEXT("Pak chunk signing mismatch on chunk [%i/%i]! Expected %s, Received %s"), ChunkInfo.Index, Signatures.ChunkHashes.Num(), *ChunkHashToString(Signatures.ChunkHashes[ChunkInfo.Index]), *ChunkHashToString(ChunkHash));
 
 		if (Signatures.DecryptedHash != Signatures.ComputeCurrentMasterHash())
 		{
diff --git a/Engine/Source/Runtime/PakFile/Public/IPlatformFilePak.h b/Engine/Source/Runtime/PakFile/Public/IPlatformFilePak.h
--- a/Engine/Source/Runtime/PakFile/Public/IPlatformFilePak.h
+++ b/Engine/Source/Runtime/PakFile/Public/IPlatformFilePak.h
@@ -34,6 +34,7 @@ typedef FSHAHash TPakChunkHash;
 #endif
 
 PAKFILE_API TPakChunkHash ComputePakChunkHash(const void* InData, int64 InDataSizeInBytes);
+PAKFILE_API FString ChunkHashToString(const TPakChunkHash& InHash);
 
 struct FPakChunkSignatureCheckFailedData
 {
