Don't mount pak files that fail the initial signature check

diff --git a/Engine/Source/Runtime/PakFile/Private/IPlatformFilePak.cpp b/Engine/Source/Runtime/PakFile/Private/IPlatformFilePak.cpp
--- a/Engine/Source/Runtime/PakFile/Private/IPlatformFilePak.cpp
+++ b/Engine/Source/Runtime/PakFile/Private/IPlatformFilePak.cpp
@@ -4252,6 +4252,11 @@ FPakFile::~FPakFile()
 	delete[] FilenameHashesIndex;
 }
 
+bool FPakFile::PassedSignatureChecks() const
+{
+	return Decryptor.IsValid() && Decryptor->IsValid();
+}
+
 FArchive* FPakFile::CreatePakReader(const TCHAR* Filename)
 {
 	FArchive* ReaderArchive = IFileManager::Get().CreateFileReader(Filename);
@@ -4278,7 +4283,16 @@ FArchive* FPakFile::SetupSignedPakReader(FArchive* ReaderArchive, const TCHAR* F
 			{
 				Decryptor = MakeUnique<FChunkCacheWorker>(ReaderArchive, Filename);
 			}
-			ReaderArchive = new FSignedArchiveReader(ReaderArchive, Decryptor.Get());
+
+			if (Decryptor->IsValid())
+			{
+				ReaderArchive = new FSignedArchiveReader(ReaderArchive, Decryptor.Get());
+			}
+			else
+			{
+				delete ReaderArchive;
+				return nullptr;
+			}
 		}
 	}
 	return ReaderArchive;
@@ -5055,7 +5069,8 @@ FArchive* FPakFile::GetSharedReader(IPlatformFile* LowerLevel)
 			}
 			if (!PakReader)
 			{
-				UE_LOG(LogPakFile, Fatal, TEXT("Unable to create pak \"%s\" handle"), *GetFilename());
+				UE_LOG(LogPakFile, Warning, TEXT("Unable to create pak \"%s\" handle"), *GetFilename());
+				return nullptr;
 			}
 
 #if DO_CHECK
@@ -5564,7 +5579,8 @@ bool FPakPlatformFile::Mount(const TCHAR* InPakFilename, uint32 PakOrder, const
 		}
 		else
 		{
-			if (Pak->GetInfo().EncryptionKeyGuid.IsValid())
+			bool bPassedSignatureChecks = !bSigned || Pak->PassedSignatureChecks();
+			if (bPassedSignatureChecks && Pak->GetInfo().EncryptionKeyGuid.IsValid())
 			{
 				UE_LOG(LogPakFile, Log, TEXT("Deferring mount of pak \"%s\" until encryption key '%s' becomes available"), InPakFilename, *Pak->GetInfo().EncryptionKeyGuid.ToString());
 
diff --git a/Engine/Source/Runtime/PakFile/Private/SignedArchiveReader.cpp b/Engine/Source/Runtime/PakFile/Private/SignedArchiveReader.cpp
--- a/Engine/Source/Runtime/PakFile/Private/SignedArchiveReader.cpp
+++ b/Engine/Source/Runtime/PakFile/Private/SignedArchiveReader.cpp
@@ -46,14 +46,20 @@ FChunkCacheWorker::FChunkCacheWorker(FArchive* InReader, const TCHAR* Filename)
 
 	Signatures.Serialize(*SigFileReader);
 	delete SigFileReader;
-	Signatures.DecryptSignatureAndValidate(Filename);
-
-	const bool bEnableMultithreading = FPlatformProcess::SupportsMultithreading();
-	if (bEnableMultithreading)
+	if (!Signatures.DecryptSignatureAndValidate(Filename))
+	{
+		// This will render this object invalid, so whoever has created this worker will have to also check IsValid() and ignore the pak afterwards
+		Signatures.ChunkHashes.Empty();
+	}
+	else
 	{
-		QueuedRequestsEvent = FPlatformProcess::GetSynchEventFromPool();
-		ChunkRequestAvailable = FPlatformProcess::GetSynchEventFromPool();
-		Thread = FRunnableThread::Create(this, TEXT("FChunkCacheWorker"), 0, TPri_BelowNormal);
+		const bool bEnableMultithreading = FPlatformProcess::SupportsMultithreading();
+		if (bEnableMultithreading)
+		{
+			QueuedRequestsEvent = FPlatformProcess::GetSynchEventFromPool();
+			ChunkRequestAvailable = FPlatformProcess::GetSynchEventFromPool();
+			Thread = FRunnableThread::Create(this, TEXT("FChunkCacheWorker"), 0, TPri_BelowNormal);
+		}
 	}
 }
 
@@ -229,30 +235,36 @@ bool FChunkCacheWorker::CheckSignature(const FChunkRequest& ChunkInfo)
 {
 	SCOPE_SECONDS_ACCUMULATOR(STAT_FChunkCacheWorker_CheckSignature);
 
-	TPakChunkHash ChunkHash;
-
-	{
-		SCOPE_SECONDS_ACCUMULATOR(STAT_FChunkCacheWorker_Serialize);
-		Reader->Seek(ChunkInfo.Offset);
-		Reader->Serialize(ChunkInfo.Buffer->Data, ChunkInfo.Size);
-	}
-	{
-		SCOPE_SECONDS_ACCUMULATOR(STAT_FChunkCacheWorker_HashBuffer);
-		ChunkHash = ComputePakChunkHash(ChunkInfo.Buffer->Data, ChunkInfo.Size);
-	}
+	bool bChunkHashesMatch = false;
 
-	bool bChunkHashesMatch = (ChunkHash == Signatures.ChunkHashes[ChunkInfo.Index]);
-	if (!bChunkHashesMatch)
+	// If our signature data wasn't validated properly on startup, we shouldn't be in here. Mark all chunk checks as failed.
+	if (ensure(IsValid()))
 	{
-		UE_LOG(LogPakFile, Warning, TEXT("Pak chunk signing mismatch on chunk [%i/%i]! Expected %s, Received %s"), ChunkInfo.Index, Signatures.ChunkHashes.Num(), *ChunkHashToString(Signatures.ChunkHashes[ChunkInfo.Index]), *ChunkHashToString(ChunkHash));
+		TPakChunkHash ChunkHash;
 
-		if (Signatures.DecryptedHash != Signatures.ComputeCurrentMasterHash())
 		{
-			UE_LOG(LogPakFile, Warning, TEXT("Master signature table has changed since initialization!"));
+			SCOPE_SECONDS_ACCUMULATOR(STAT_FChunkCacheWorker_Serialize);
+			Reader->Seek(ChunkInfo.Offset);
+			Reader->Serialize(ChunkInfo.Buffer->Data, ChunkInfo.Size);
 		}
+		{
+			SCOPE_SECONDS_ACCUMULATOR(STAT_FChunkCacheWorker_HashBuffer);
+			ChunkHash = ComputePakChunkHash(ChunkInfo.Buffer->Data, ChunkInfo.Size);
+		}
+
+		bChunkHashesMatch = IsValid() && (ChunkHash == Signatures.ChunkHashes[ChunkInfo.Index]);
+		if (!bChunkHashesMatch)
+		{
+			UE_LOG(LogPakFile, Warning, TEXT("Pak chunk signing mismatch on chunk [%i/%i]! Expected 0x%8X, Received 0x%8X"), ChunkInfo.Index, Signatures.ChunkHashes.Num(), *LexToString(Signatures.ChunkHashes[ChunkInfo.Index]), *LexToString(ChunkHash));
 
-		const FPakChunkSignatureCheckFailedData Data(Reader->GetArchiveName(), Signatures.ChunkHashes[ChunkInfo.Index], ChunkHash, ChunkInfo.Index);
-		FPakPlatformFile::GetPakChunkSignatureCheckFailedHandler().Broadcast(Data);
+			if (Signatures.DecryptedHash != Signatures.ComputeCurrentMasterHash())
+			{
+				UE_LOG(LogPakFile, Warning, TEXT("Master signature table has changed since initialization!"));
+			}
+
+			const FPakChunkSignatureCheckFailedData Data(Reader->GetArchiveName(), Signatures.ChunkHashes[ChunkInfo.Index], ChunkHash, ChunkInfo.Index);
+			FPakPlatformFile::GetPakChunkSignatureCheckFailedHandler().Broadcast(Data);
+		}
 	}
 	
 	return bChunkHashesMatch;
@@ -301,6 +313,11 @@ void FChunkCacheWorker::FlushRemainingChunkCompletionEvents()
 	}
 }
 
+bool FChunkCacheWorker::IsValid() const
+{
+	return Signatures.ChunkHashes.Num() > 0;
+}
+
 void FChunkCacheWorker::ReleaseChunk(FChunkRequest& Chunk)
 {
 	if (Chunk.RefCount.Decrement() == 0 && QueuedRequestsEvent != nullptr)
diff --git a/Engine/Source/Runtime/PakFile/Private/SignedArchiveReader.h b/Engine/Source/Runtime/PakFile/Private/SignedArchiveReader.h
--- a/Engine/Source/Runtime/PakFile/Private/SignedArchiveReader.h
+++ b/Engine/Source/Runtime/PakFile/Private/SignedArchiveReader.h
@@ -201,6 +201,13 @@ public:
 	 */
 	void ReleaseChunk(FChunkRequest& Chunk);
 
+
+	/**
+	* Indicates that this chunk worker is valid. If the signature file couldn't be loaded or if it failed
+	* the master table check, this will be false
+	*/
+	bool IsValid() const;
+
 	friend class FSignedArchiveReader;
 };
 
diff --git a/Engine/Source/Runtime/PakFile/Public/IPlatformFilePak.h b/Engine/Source/Runtime/PakFile/Public/IPlatformFilePak.h
--- a/Engine/Source/Runtime/PakFile/Public/IPlatformFilePak.h
+++ b/Engine/Source/Runtime/PakFile/Public/IPlatformFilePak.h
@@ -608,6 +608,13 @@ public:
 		return bIsValid;
 	}
 
+	/**
+	 * Checks if the pak has valid chunk signature checking data, and that the data passed the initial signing check
+	 *
+	 * @return true if this pak file has passed the initial signature checking phase
+	 */
+	bool PassedSignatureChecks() const;
+
 	/**
 	 * Gets pak filename.
 	 *
