UnloadPakEntryFilenames uses 64bit hashes.
#rnx

diff --git a/Engine/Source/Runtime/PakFile/Private/IPlatformFilePak.cpp b/Engine/Source/Runtime/PakFile/Private/IPlatformFilePak.cpp
--- a/Engine/Source/Runtime/PakFile/Private/IPlatformFilePak.cpp
+++ b/Engine/Source/Runtime/PakFile/Private/IPlatformFilePak.cpp
@@ -4513,7 +4513,7 @@ bool FPakFile::Check()
 
 struct FMiniFileEntry
 {
-	uint32 FilenameHash;
+	uint64 FilenameHash;
 	int32 EntryIndex;
 };
 
@@ -4532,7 +4532,7 @@ static inline int32 CDECL CompareFMiniFileEntry(const void* Left, const void* Ri
 	return 0;
 }
 
-bool FPakFile::UnloadPakEntryFilenames(TMap<uint32, FPakEntry>& CrossPakCollisionChecker, TArray<FString>* DirectoryRootsToKeep)
+bool FPakFile::UnloadPakEntryFilenames(TMap<uint64, FPakEntry>& CrossPakCollisionChecker, TArray<FString>* DirectoryRootsToKeep)
 {
 	// If the process has already been done, get out of here.
 	if (bFilenamesRemoved)
@@ -4558,7 +4558,7 @@ bool FPakFile::UnloadPakEntryFilenames(TMap<uint32, FPakEntry>& CrossPakCollisio
 		// No collisions yet for this pass.
 		bHasCollision = false;
 
-		TMap<uint32, FPakEntry> NewCollisionCheckEntries;
+		TMap<uint64, FPakEntry> NewCollisionCheckEntries;
 
 		// Build the list of hashes from the Index based on the starting hash.
 		int32 EntryIndex = 0;
@@ -4566,8 +4566,8 @@ bool FPakFile::UnloadPakEntryFilenames(TMap<uint32, FPakEntry>& CrossPakCollisio
 		{
 			for (FPakDirectory::TConstIterator DirectoryIt(It.Value()); DirectoryIt; ++DirectoryIt)
 			{
-				FString FinalFilename = (It.Key() / DirectoryIt.Key()).ToLower();
-				uint32 FilenameHash = FFnv::MemFnv32(*FinalFilename, FinalFilename.Len() * sizeof(TCHAR), FilenameStartHash);
+				FString FinalFilename = It.Key() / DirectoryIt.Key();
+				uint64 FilenameHash = FFnv::MemFnv64(*FinalFilename.ToLower(), FinalFilename.Len() * sizeof(TCHAR), FilenameStartHash);
 				MiniFileEntries[EntryIndex].FilenameHash = FilenameHash;
 				MiniFileEntries[EntryIndex].EntryIndex = DirectoryIt.Value();
 				++EntryIndex;
@@ -4619,7 +4619,7 @@ bool FPakFile::UnloadPakEntryFilenames(TMap<uint32, FPakEntry>& CrossPakCollisio
 
 	// Allocate the storage space.
 	FilenameHashesIndices = new int32[NumEntries];
-	FilenameHashes = new uint32[NumEntries];
+	FilenameHashes = new uint64[NumEntries];
 	int32 LastHashMostSignificantBits = -1;
 
 	// FilenameHashesIndex provides small 'arenas' of binary searchable filename hashes.
@@ -4645,8 +4645,8 @@ bool FPakFile::UnloadPakEntryFilenames(TMap<uint32, FPakEntry>& CrossPakCollisio
 	{
 		// If a new index entry is needed as a result of crossing over into a larger hash group
 		// as specified through the 8 most significant bits of the hash, store the entry index.
-		uint32 FilenameHash = MiniFileEntries[EntryIndex].FilenameHash;
-		int32 HashMostSignificantBits = FilenameHash >> 24;
+		uint64 FilenameHash = MiniFileEntries[EntryIndex].FilenameHash;
+		int32 HashMostSignificantBits = FilenameHash >> 56;
 		if (HashMostSignificantBits != LastHashMostSignificantBits)
 		{
 			for (int32 BitsIndex = LastHashMostSignificantBits + 1; BitsIndex <= HashMostSignificantBits; ++BitsIndex)
@@ -5109,12 +5109,12 @@ FPakFile::EFindResult FPakFile::Find(const FString& Filename, FPakEntry* OutEntr
 				++SplitStartPtr;
 				--SplitLen;
 			}
-            uint32 PathHash = FFnv::MemFnv32(SplitStartPtr, SplitLen * sizeof(TCHAR), FilenameStartHash);
+			uint64 PathHash = FFnv::MemFnv64(SplitStartPtr, SplitLen * sizeof(TCHAR), FilenameStartHash);
 
 			// Look it up in our sorted-by-filename-hash array.
-			uint32 PathHashMostSignificantBits = PathHash >> 24;
+			uint32 PathHashMostSignificantBits = PathHash >> 56;
 			uint32 HashEntriesCount = FilenameHashesIndex[PathHashMostSignificantBits + 1] - FilenameHashesIndex[PathHashMostSignificantBits];
-			uint32* FoundHash = (uint32*)bsearch(&PathHash, FilenameHashes + FilenameHashesIndex[PathHashMostSignificantBits], HashEntriesCount, sizeof(uint32), CompareFilenameHashes);
+			uint64* FoundHash = (uint64*)bsearch(&PathHash, FilenameHashes + FilenameHashesIndex[PathHashMostSignificantBits], HashEntriesCount, sizeof(uint64), CompareFilenameHashes);
 			if (FoundHash != NULL)
 			{
 				bool bDeleted = false;
@@ -5972,6 +5972,7 @@ bool FPakPlatformFile::CopyFile(const TCHAR* To, const TCHAR* From, EPlatformFil
 
 void FPakPlatformFile::UnloadPakEntryFilenames(TArray<FString>* DirectoryRootsToKeep)
 {
+	int32 NumFiles = 0;
 	double Timer = 0.0;
 	{
 		SCOPE_SECONDS_COUNTER(Timer);
@@ -5979,9 +5980,8 @@ void FPakPlatformFile::UnloadPakEntryFilenames(TArray<FString>* DirectoryRootsTo
 		TArray<FPakListEntry> Paks;
 		GetMountedPaks(Paks);
 		FPlatformMisc::LowLevelOutputDebugStringf(TEXT("Unloading Pak Entry Filenames\n"));
-		TMap<uint32, FPakEntry> CrossPakCollisionDetector;
+		TMap<uint64, FPakEntry> CrossPakCollisionDetector;
 
-		int32 NumFiles = 0;
 		for (const FPakListEntry& Pak : Paks)
 		{
 			NumFiles += Pak.PakFile->GetNumFiles();
@@ -5998,7 +5998,7 @@ void FPakPlatformFile::UnloadPakEntryFilenames(TArray<FString>* DirectoryRootsTo
 			UE_CLOG(bEncounteredCollisions, LogPakFile, Warning, TEXT("Encountered name collisions while unloading pak index strings for %s"), *Pak.PakFile->GetFilename());
 		}
 	}
-	UE_LOG(LogPakFile, Log, TEXT("PakEntry filenames unloaded in %.4fs"), Timer);
+	UE_LOG(LogPakFile, Log, TEXT("PakEntry filenames unloaded %d filenames in %.4fs"), NumFiles, Timer);
 }
 
 void FPakPlatformFile::ShrinkPakEntriesMemoryUsage()
diff --git a/Engine/Source/Runtime/PakFile/Public/IPlatformFilePak.h b/Engine/Source/Runtime/PakFile/Public/IPlatformFilePak.h
--- a/Engine/Source/Runtime/PakFile/Public/IPlatformFilePak.h
+++ b/Engine/Source/Runtime/PakFile/Public/IPlatformFilePak.h
@@ -517,13 +517,13 @@ class PAKFILE_API FPakFile : FNoncopyable
 	/** Pak Index organized as a map of directories for faster Directory iteration. Completely valid only when bFilenamesRemoved == false, although portions may still be valid after a call to UnloadPakEntryFilenames() while utilizing DirectoryRootsToKeep. */
 	TMap<FString, FPakDirectory> Index;
 	/** The hash to use when generating a filename hash (CRC) to avoid collisions within the hashed filename space. */
-	uint32 FilenameStartHash;
+	uint64 FilenameStartHash;
 	/** An array of 256 + 1 size that represents the starting index of the most significant byte of a hash group within the FilenameHashes array. */
 	uint32* FilenameHashesIndex;
 	/** An array of NumEntries size mapping 1:1 with FilenameHashes and describing the index of the FPakEntry. */
 	int32* FilenameHashesIndices;
 	/** A tightly packed array of filename hashes (CRC) of NumEntries size. */
-	uint32* FilenameHashes;
+	uint64* FilenameHashes;
 	/** A tightly packed array, NumEntries in size, of offsets to the pak entry data within the MiniPakEntries buffer */
 	uint32* MiniPakEntriesOffsets;
 	/** Memory buffer representing the minimal file entry headers, NumEntries in size */
@@ -549,8 +549,8 @@ class PAKFILE_API FPakFile : FNoncopyable
 
 	static inline int32 CDECL CompareFilenameHashes(const void* Left, const void* Right)
 	{
-		const uint32* LeftHash = (const uint32*)Left;
-		const uint32* RightHash = (const uint32*)Right;
+		const uint64* LeftHash = (const uint64*)Left;
+		const uint64* RightHash = (const uint64*)Right;
 		if (*LeftHash < *RightHash)
 		{
 			return -1;
@@ -945,7 +945,7 @@ public:
 	 * @param CrossPakCollisionChecker A map of hash->fileentry records encountered during filename unloading on other pak files. Used to detect collisions with entries in other pak files.
 	 * @param DirectoryRootsToKeep An array of strings in wildcard format that specify whole directory structures of filenames to keep in memory for directory iteration to work.
 	 */
-	bool UnloadPakEntryFilenames(TMap<uint32, FPakEntry>& CrossPakCollisionChecker, TArray<FString>* DirectoryRootsToKeep = nullptr);
+	bool UnloadPakEntryFilenames(TMap<uint64, FPakEntry>& CrossPakCollisionChecker, TArray<FString>* DirectoryRootsToKeep = nullptr);
 
 	/**
 	 * Lower memory usage by bit-encoding the pak file entry information.
