Improve logging of pak signature failures.
Treat a missing .sig file like a signature failure. Log it and ignore the pak file.

diff --git a/Engine/Source/Runtime/PakFile/Private/IPlatformFilePak.cpp b/Engine/Source/Runtime/PakFile/Private/IPlatformFilePak.cpp
--- a/Engine/Source/Runtime/PakFile/Private/IPlatformFilePak.cpp
+++ b/Engine/Source/Runtime/PakFile/Private/IPlatformFilePak.cpp
@@ -5582,7 +5582,7 @@ bool FPakPlatformFile::Mount(const TCHAR* InPakFilename, uint32 PakOrder, const
 			bool bPassedSignatureChecks = !bSigned || Pak->PassedSignatureChecks();
 			if (bPassedSignatureChecks && Pak->GetInfo().EncryptionKeyGuid.IsValid())
 			{
-				UE_LOG(LogPakFile, Log, TEXT("Deferring mount of pak \"%s\" until encryption key '%s' becomes available"), InPakFilename, *Pak->GetInfo().EncryptionKeyGuid.ToString());
+				UE_LOG(LogPakFile, Display, TEXT("Deferring mount of pak \"%s\" until encryption key '%s' becomes available"), InPakFilename, *Pak->GetInfo().EncryptionKeyGuid.ToString());
 
 				check(!GetRegisteredEncryptionKeys().HasKey(Pak->GetInfo().EncryptionKeyGuid));
 				FPakListDeferredEntry& Entry = PendingEncryptedPakFiles[PendingEncryptedPakFiles.Add(FPakListDeferredEntry())];
diff --git a/Engine/Source/Runtime/PakFile/Private/SignedArchiveReader.cpp b/Engine/Source/Runtime/PakFile/Private/SignedArchiveReader.cpp
--- a/Engine/Source/Runtime/PakFile/Private/SignedArchiveReader.cpp
+++ b/Engine/Source/Runtime/PakFile/Private/SignedArchiveReader.cpp
@@ -39,27 +39,30 @@ FChunkCacheWorker::FChunkCacheWorker(FArchive* InReader, const TCHAR* Filename)
 	FString SigFileFilename = FPaths::ChangeExtension(Filename, TEXT("sig"));
 	FArchive* SigFileReader = IFileManager::Get().CreateFileReader(*SigFileFilename);
 
-	if (SigFileReader == nullptr)
+	if (SigFileReader != nullptr)
 	{
-		UE_LOG(LogPakFile, Fatal, TEXT("Couldn't find pak signature file '%s'"), *SigFileFilename);
-	}
-
-	Signatures.Serialize(*SigFileReader);
-	delete SigFileReader;
-	if (!Signatures.DecryptSignatureAndValidate(Filename))
-	{
-		// This will render this object invalid, so whoever has created this worker will have to also check IsValid() and ignore the pak afterwards
-		Signatures.ChunkHashes.Empty();
+		Signatures.Serialize(*SigFileReader);
+		delete SigFileReader;
+		if (!Signatures.DecryptSignatureAndValidate(Filename))
+		{
+			// This will render this object invalid, so whoever has created this worker will have to also check IsValid() and ignore the pak afterwards
+			Signatures.ChunkHashes.Empty();
+		}
+		else
+		{
+			const bool bEnableMultithreading = FPlatformProcess::SupportsMultithreading();
+			if (bEnableMultithreading)
+			{
+				QueuedRequestsEvent = FPlatformProcess::GetSynchEventFromPool();
+				ChunkRequestAvailable = FPlatformProcess::GetSynchEventFromPool();
+				Thread = FRunnableThread::Create(this, TEXT("FChunkCacheWorker"), 0, TPri_BelowNormal);
+			}
+		}
 	}
 	else
 	{
-		const bool bEnableMultithreading = FPlatformProcess::SupportsMultithreading();
-		if (bEnableMultithreading)
-		{
-			QueuedRequestsEvent = FPlatformProcess::GetSynchEventFromPool();
-			ChunkRequestAvailable = FPlatformProcess::GetSynchEventFromPool();
-			Thread = FRunnableThread::Create(this, TEXT("FChunkCacheWorker"), 0, TPri_BelowNormal);
-		}
+		UE_LOG(LogPakFile, Warning, TEXT("Couldn't find pak signature file '%s'"), *SigFileFilename);
+		FPakPlatformFile::GetPakMasterSignatureTableCheckFailureHandler().Broadcast(Filename);
 	}
 }
 
@@ -255,7 +258,7 @@ bool FChunkCacheWorker::CheckSignature(const FChunkRequest& ChunkInfo)
 		bChunkHashesMatch = IsValid() && (ChunkHash == Signatures.ChunkHashes[ChunkInfo.Index]);
 		if (!bChunkHashesMatch)
 		{
-			UE_LOG(LogPakFile, Warning, TEXT("Pak chunk signing mismatch on chunk [%i/%i]! Expected 0x%8X, Received 0x%8X"), ChunkInfo.Index, Signatures.ChunkHashes.Num(), *LexToString(Signatures.ChunkHashes[ChunkInfo.Index]), *LexToString(ChunkHash));
+			UE_LOG(LogPakFile, Warning, TEXT("Pak chunk signing mismatch on chunk [%i/%i]! Expected %s, Received %s"), ChunkInfo.Index, Signatures.ChunkHashes.Num(), *ChunkHashToString(Signatures.ChunkHashes[ChunkInfo.Index]), *ChunkHashToString(ChunkHash));
 
 			if (Signatures.DecryptedHash != Signatures.ComputeCurrentMasterHash())
 			{
diff --git a/Engine/Source/Runtime/PakFile/Public/IPlatformFilePak.h b/Engine/Source/Runtime/PakFile/Public/IPlatformFilePak.h
--- a/Engine/Source/Runtime/PakFile/Public/IPlatformFilePak.h
+++ b/Engine/Source/Runtime/PakFile/Public/IPlatformFilePak.h
@@ -2441,6 +2441,7 @@ struct FPakSignatureFile
 		FSHAHash CurrentHash = ComputeCurrentMasterHash();
 		if (DecryptedHash != CurrentHash)
 		{
+			UE_LOG(LogPakFile, Warning, TEXT("Pak signature table validation failed for '%s'! Expected %s, Received %s"), *InFilename, *DecryptedHash.ToString(), *CurrentHash.ToString());
 			FPakPlatformFile::GetPakMasterSignatureTableCheckFailureHandler().Broadcast(InFilename);
 			return false;
 		}
