Centralize the loading of pak signature data so that the same object can be reused between the SignedArchiveReader and FPakPrecacher implementations. Stops the same data being loaded and verified twice

diff --git a/Engine/Source/Runtime/PakFile/Private/IPlatformFilePak.cpp b/Engine/Source/Runtime/PakFile/Private/IPlatformFilePak.cpp
--- a/Engine/Source/Runtime/PakFile/Private/IPlatformFilePak.cpp
+++ b/Engine/Source/Runtime/PakFile/Private/IPlatformFilePak.cpp
@@ -245,6 +245,51 @@ TSharedPtr<FRSA::FKey, ESPMode::ThreadSafe> FPakPlatformFile::GetPakSigningKey()
 	return Key;
 }
 
+TMap<FName, TSharedPtr<const struct FPakSignatureFile, ESPMode::ThreadSafe>> FPakPlatformFile::PakSignatureFileCache;
+FCriticalSection FPakPlatformFile::PakSignatureFileCacheLock;
+TSharedPtr<const struct FPakSignatureFile, ESPMode::ThreadSafe> FPakPlatformFile::GetPakSignatureFile(const TCHAR* InFilename)
+{
+	FScopeLock Lock(&PakSignatureFileCacheLock);
+	FName FilenameFName(InFilename);
+	if (const TSharedPtr<const struct FPakSignatureFile, ESPMode::ThreadSafe>* SignaturesFile = PakSignatureFileCache.Find(FilenameFName))
+	{
+		return *SignaturesFile;
+	}
+
+	TSharedPtr<FPakSignatureFile, ESPMode::ThreadSafe> SignaturesFile = MakeShared<FPakSignatureFile, ESPMode::ThreadSafe>();
+	FString SignaturesFilename = FPaths::ChangeExtension(InFilename, TEXT("sig"));
+	FArchive* Reader = IFileManager::Get().CreateFileReader(*SignaturesFilename);
+	if (Reader != nullptr)
+	{
+		SignaturesFile->Serialize(*Reader);
+		delete Reader;
+
+		if (!SignaturesFile->DecryptSignatureAndValidate(InFilename))
+		{
+			SignaturesFile.Reset();
+		}
+
+		PakSignatureFileCache.Add(FilenameFName, SignaturesFile);
+	}
+	else
+	{
+		UE_LOG(LogPakFile, Warning, TEXT("Couldn't find pak signature file '%s'"), InFilename);
+		FPakPlatformFile::GetPakMasterSignatureTableCheckFailureHandler().Broadcast(InFilename);
+	}
+
+	return SignaturesFile;
+}
+
+void FPakPlatformFile::RemoveCachedPakSignaturesFile(const TCHAR* InFilename)
+{
+	FScopeLock Lock(&PakSignatureFileCacheLock);
+	FName FilenameFName(InFilename);
+	if (PakSignatureFileCache.Contains(FilenameFName))
+	{
+		PakSignatureFileCache.Remove(FilenameFName);
+	}
+}
+
 DECLARE_DWORD_ACCUMULATOR_STAT(TEXT("PakCache Sync Decrypts (Uncompressed Path)"), STAT_PakCache_SyncDecrypts, STATGROUP_PakFile);
 DECLARE_FLOAT_ACCUMULATOR_STAT(TEXT("PakCache Decrypt Time"), STAT_PakCache_DecryptTime, STATGROUP_PakFile);
 DECLARE_DWORD_ACCUMULATOR_STAT(TEXT("PakCache Async Decrypts (Compressed Path)"), STAT_PakCache_CompressedDecrypts, STATGROUP_PakFile);
@@ -1057,7 +1102,7 @@ class FPakPrecacher
 		TIntervalTreeIndex InRequests[AIOP_NUM][(int32)EInRequestStatus::Num];
 		TIntervalTreeIndex CacheBlocks[(int32)EBlockStatus::Num];
 
-		FPakSignatureFile Signatures;
+		TSharedPtr<const FPakSignatureFile, ESPMode::ThreadSafe> Signatures;
 
 		FPakData(IAsyncReadFileHandle* InHandle, FName InName, int64 InTotalSize)
 			: Handle(InHandle)
@@ -1066,6 +1111,7 @@ class FPakPrecacher
 			, MaxShift(0)
 			, BytesToBitsShift(0)
 			, Name(InName)
+			, Signatures(nullptr)
 		{
 			check(Handle && TotalSize > 0 && Name != NAME_None);
 			for (int32 Index = 0; Index < AIOP_NUM; Index++)
@@ -1244,17 +1290,16 @@ public:
 			if (SigningKey.IsValid())
 			{
 				// Load signature data
-				FString SignaturesFilename = FPaths::ChangeExtension(*PakFilename, TEXT("sig"));
-				IFileHandle* SignaturesFile = LowerLevel->OpenRead(*SignaturesFilename);
-				ensure(SignaturesFile);
-				FArchiveFileReaderGeneric* Reader = new FArchiveFileReaderGeneric(SignaturesFile, *SignaturesFilename, SignaturesFile->Size());
-				Pak.Signatures.Serialize(*Reader);
-				delete Reader;
-				Pak.Signatures.DecryptSignatureAndValidate(SigningKey, PakFilename);
-
+				Pak.Signatures = FPakPlatformFile::GetPakSignatureFile(*PakFilename);
+				
+				// We should never get here unless the signature file exists and is validated. The original FPakFile creation
+				// on the main thread would have failed and the pak would never have been mounted otherwise, and then we would
+				// never have issued read requests to the pak precacher.
+				check(Pak.Signatures);
+				
 				// Check that we have the correct match between signature and pre-cache granularity
 				int64 NumPakChunks = Align(PakFileSize, FPakInfo::MaxChunkDataSize) / FPakInfo::MaxChunkDataSize;
-				ensure(NumPakChunks == Pak.Signatures.ChunkHashes.Num());
+				ensure(NumPakChunks == Pak.Signatures->ChunkHashes.Num());
 			}
 		}
 		return PakIndexPtr;
@@ -1267,7 +1312,7 @@ public:
 
 		for (FPakData& PakData : CachedPakData)
 		{
-			for (TPakChunkHash& Hash : PakData.Signatures.ChunkHashes)
+			for (const TPakChunkHash& Hash : PakData.Signatures->ChunkHashes)
 			{
 				*((uint8*)&Hash) |= 0x1;
 			}
@@ -3150,11 +3195,11 @@ void FPakPrecacher::DoSignatureCheck(bool bWasCanceled, IAsyncReadRequest* Reque
 		SignatureIndex = RequestOffset / FPakInfo::MaxChunkDataSize;
 
 		FPakData& PakData = CachedPakData[PakIndex];
-		MasterSignatureHash = PakData.Signatures.DecryptedHash;
+		MasterSignatureHash = PakData.Signatures->DecryptedHash;
 
 		for (int32 CacheIndex = 0; CacheIndex < FMath::Min(NumSignaturesToCheck, MaxHashesToCache); ++CacheIndex)
 		{
-			HashCache[CacheIndex] = PakData.Signatures.ChunkHashes[SignatureIndex + CacheIndex];
+			HashCache[CacheIndex] = PakData.Signatures->ChunkHashes[SignatureIndex + CacheIndex];
 		}
 	}
 
@@ -3174,7 +3219,7 @@ void FPakPrecacher::DoSignatureCheck(bool bWasCanceled, IAsyncReadRequest* Reque
 			FPakData& PakData = CachedPakData[PakIndex];
 			for (int32 CacheIndex = 0; (CacheIndex < MaxHashesToCache) && ((SignedChunkIndex + CacheIndex) < NumSignaturesToCheck); ++CacheIndex)
 			{
-				HashCache[CacheIndex] = PakData.Signatures.ChunkHashes[SignatureIndex + CacheIndex];
+				HashCache[CacheIndex] = PakData.Signatures->ChunkHashes[SignatureIndex + CacheIndex];
 			}
 		}
 
@@ -3189,10 +3234,10 @@ void FPakPrecacher::DoSignatureCheck(bool bWasCanceled, IAsyncReadRequest* Reque
 				FScopeLock Lock(&CachedFilesScopeLock);
 				FPakData* PakData = &CachedPakData[PakIndex];
 
-				UE_LOG(LogPakFile, Warning, TEXT("Pak chunk signing mismatch on chunk [%i/%i]! Expected %s, Received %s"), SignatureIndex, PakData->Signatures.ChunkHashes.Num(), *ChunkHashToString(PakData->Signatures.ChunkHashes[SignatureIndex]), *ChunkHashToString(ThisHash));
+				UE_LOG(LogPakFile, Warning, TEXT("Pak chunk signing mismatch on chunk [%i/%i]! Expected %s, Received %s"), SignatureIndex, PakData->Signatures->ChunkHashes.Num(), *ChunkHashToString(PakData->Signatures->ChunkHashes[SignatureIndex]), *ChunkHashToString(ThisHash));
 
 				// Check the signatures are still as we expected them
-				if (PakData->Signatures.DecryptedHash != PakData->Signatures.ComputeCurrentMasterHash())
+				if (PakData->Signatures->DecryptedHash != PakData->Signatures->ComputeCurrentMasterHash())
 				{
 					UE_LOG(LogPakFile, Warning, TEXT("Master signature table has changed since initialization!"));
 				}
@@ -5624,7 +5669,9 @@ bool FPakPlatformFile::Unmount(const TCHAR* InPakFilename)
 		{
 			if (PakFiles[PakIndex].PakFile->GetFilename() == InPakFilename)
 			{
-				delete PakFiles[PakIndex].PakFile;
+				FPakListEntry& PakListEntry = PakFiles[PakIndex];
+				RemoveCachedPakSignaturesFile(*PakListEntry.PakFile->GetFilename());
+				delete PakListEntry.PakFile;
 				PakFiles.RemoveAt(PakIndex);
 				return true;
 			}
diff --git a/Engine/Source/Runtime/PakFile/Private/SignedArchiveReader.cpp b/Engine/Source/Runtime/PakFile/Private/SignedArchiveReader.cpp
--- a/Engine/Source/Runtime/PakFile/Private/SignedArchiveReader.cpp
+++ b/Engine/Source/Runtime/PakFile/Private/SignedArchiveReader.cpp
@@ -36,33 +36,17 @@ FChunkCacheWorker::FChunkCacheWorker(FArchive* InReader, const TCHAR* Filename)
 	, QueuedRequestsEvent(nullptr)
 	, ChunkRequestAvailable(nullptr)
 {
-	FString SigFileFilename = FPaths::ChangeExtension(Filename, TEXT("sig"));
-	FArchive* SigFileReader = IFileManager::Get().CreateFileReader(*SigFileFilename);
+	Signatures = FPakPlatformFile::GetPakSignatureFile(Filename);
 
-	if (SigFileReader != nullptr)
+	if (Signatures.IsValid())
 	{
-		Signatures.Serialize(*SigFileReader);
-		delete SigFileReader;
-		if (!Signatures.DecryptSignatureAndValidate(Filename))
+		const bool bEnableMultithreading = FPlatformProcess::SupportsMultithreading();
+		if (bEnableMultithreading)
 		{
-			// This will render this object invalid, so whoever has created this worker will have to also check IsValid() and ignore the pak afterwards
-			Signatures.ChunkHashes.Empty();
+			QueuedRequestsEvent = FPlatformProcess::GetSynchEventFromPool();
+			ChunkRequestAvailable = FPlatformProcess::GetSynchEventFromPool();
+			Thread = FRunnableThread::Create(this, TEXT("FChunkCacheWorker"), 0, TPri_BelowNormal);
 		}
-		else
-		{
-			const bool bEnableMultithreading = FPlatformProcess::SupportsMultithreading();
-			if (bEnableMultithreading)
-			{
-				QueuedRequestsEvent = FPlatformProcess::GetSynchEventFromPool();
-				ChunkRequestAvailable = FPlatformProcess::GetSynchEventFromPool();
-				Thread = FRunnableThread::Create(this, TEXT("FChunkCacheWorker"), 0, TPri_BelowNormal);
-			}
-		}
-	}
-	else
-	{
-		UE_LOG(LogPakFile, Warning, TEXT("Couldn't find pak signature file '%s'"), *SigFileFilename);
-		FPakPlatformFile::GetPakMasterSignatureTableCheckFailureHandler().Broadcast(Filename);
 	}
 }
 
@@ -239,6 +223,7 @@ bool FChunkCacheWorker::CheckSignature(const FChunkRequest& ChunkInfo)
 	SCOPE_SECONDS_ACCUMULATOR(STAT_FChunkCacheWorker_CheckSignature);
 
 	bool bChunkHashesMatch = false;
+	check(Signatures.IsValid());
 
 	// If our signature data wasn't validated properly on startup, we shouldn't be in here. Mark all chunk checks as failed.
 	if (ensure(IsValid()))
@@ -255,17 +240,17 @@ bool FChunkCacheWorker::CheckSignature(const FChunkRequest& ChunkInfo)
 			ChunkHash = ComputePakChunkHash(ChunkInfo.Buffer->Data, ChunkInfo.Size);
 		}
 
-		bChunkHashesMatch = IsValid() && (ChunkHash == Signatures.ChunkHashes[ChunkInfo.Index]);
+		bChunkHashesMatch = IsValid() && (ChunkHash == Signatures->ChunkHashes[ChunkInfo.Index]);
 		if (!bChunkHashesMatch)
 		{
-			UE_LOG(LogPakFile, Warning, TEXT("Pak chunk signing mismatch on chunk [%i/%i]! Expected %s, Received %s"), ChunkInfo.Index, Signatures.ChunkHashes.Num(), *ChunkHashToString(Signatures.ChunkHashes[ChunkInfo.Index]), *ChunkHashToString(ChunkHash));
+			UE_LOG(LogPakFile, Warning, TEXT("Pak chunk signing mismatch on chunk [%i/%i]! Expected %s, Received %s"), ChunkInfo.Index, Signatures->ChunkHashes.Num(), *ChunkHashToString(Signatures->ChunkHashes[ChunkInfo.Index]), *ChunkHashToString(ChunkHash));
 
-			if (Signatures.DecryptedHash != Signatures.ComputeCurrentMasterHash())
+			if (Signatures->DecryptedHash != Signatures->ComputeCurrentMasterHash())
 			{
 				UE_LOG(LogPakFile, Warning, TEXT("Master signature table has changed since initialization!"));
 			}
 
-			const FPakChunkSignatureCheckFailedData Data(Reader->GetArchiveName(), Signatures.ChunkHashes[ChunkInfo.Index], ChunkHash, ChunkInfo.Index);
+			const FPakChunkSignatureCheckFailedData Data(Reader->GetArchiveName(), Signatures->ChunkHashes[ChunkInfo.Index], ChunkHash, ChunkInfo.Index);
 			FPakPlatformFile::GetPakChunkSignatureCheckFailedHandler().Broadcast(Data);
 		}
 	}
@@ -318,7 +303,7 @@ void FChunkCacheWorker::FlushRemainingChunkCompletionEvents()
 
 bool FChunkCacheWorker::IsValid() const
 {
-	return Signatures.ChunkHashes.Num() > 0;
+	return Signatures.IsValid() && (Signatures->ChunkHashes.Num() > 0);
 }
 
 void FChunkCacheWorker::ReleaseChunk(FChunkRequest& Chunk)
diff --git a/Engine/Source/Runtime/PakFile/Private/SignedArchiveReader.h b/Engine/Source/Runtime/PakFile/Private/SignedArchiveReader.h
--- a/Engine/Source/Runtime/PakFile/Private/SignedArchiveReader.h
+++ b/Engine/Source/Runtime/PakFile/Private/SignedArchiveReader.h
@@ -105,7 +105,7 @@ class FChunkCacheWorker : public FRunnable
 	};
 
 	/** Reference hashes */
-	FPakSignatureFile Signatures;
+	TSharedPtr<const FPakSignatureFile, ESPMode::ThreadSafe> Signatures;
 	/** Hash of the sig file data. Used to check that nothing was corrupted when a signature check fails */
 	TPakChunkHash OriginalSignatureFileHash;
 	/** Thread to run the worker FRunnable on */
diff --git a/Engine/Source/Runtime/PakFile/Public/IPlatformFilePak.h b/Engine/Source/Runtime/PakFile/Public/IPlatformFilePak.h
--- a/Engine/Source/Runtime/PakFile/Public/IPlatformFilePak.h
+++ b/Engine/Source/Runtime/PakFile/Public/IPlatformFilePak.h
@@ -1545,6 +1545,18 @@ public:
 	*/
 	static TSharedPtr<FRSA::FKey, ESPMode::ThreadSafe> GetPakSigningKey();
 
+	/**
+	* Load a pak signature file. Validates the contents by comparing a SHA hash of the chunk table against and encrypted version that
+	* is stored within the file. Returns nullptr if the data is missing or fails the signature check. This function also calls
+	* the generic pak signature failure delegates if anything is wrong.
+	*/
+	static TSharedPtr<const struct FPakSignatureFile, ESPMode::ThreadSafe> GetPakSignatureFile(const TCHAR* InFilename);
+
+	/**
+	 * Remove the intenrally cached pointer to the signature file for the specified pak
+	 */
+	static void RemoveCachedPakSignaturesFile(const TCHAR* InFilename);
+
 	/**
 	 * Constructor.
 	 * 
@@ -2354,6 +2366,10 @@ public:
 	static void TrackPak(const TCHAR* Filename, const FPakEntry* PakEntry);
 	static TMap<FString, int32>& GetPakMap() { return GPakSizeMap; }
 #endif
+
+	// Internal cache of pak signature files
+	static TMap<FName, TSharedPtr<const struct FPakSignatureFile, ESPMode::ThreadSafe>> PakSignatureFileCache;
+	static FCriticalSection PakSignatureFileCacheLock;
 };
 
 /**
