Make ChunkHashToString() an inline function in IPlatformFilePak.h so that it can be used by other modules without requiring a link dependency

diff --git a/Engine/Source/Runtime/PakFile/Private/IPlatformFilePak.cpp b/Engine/Source/Runtime/PakFile/Private/IPlatformFilePak.cpp
--- a/Engine/Source/Runtime/PakFile/Private/IPlatformFilePak.cpp
+++ b/Engine/Source/Runtime/PakFile/Private/IPlatformFilePak.cpp
@@ -158,15 +158,6 @@ TPakChunkHash ComputePakChunkHash(const void* InData, int64 InDataSizeInBytes)
 #endif
 }
 
-FString ChunkHashToString(const TPakChunkHash& InHash)
-{
-#if PAKHASH_USE_CRC
-	return FString::Printf(TEXT("%08X"), InHash);
-#else
-	return LexToString(InHash);
-#endif
-}
-
 #ifndef EXCLUDE_NONPAK_UE_EXTENSIONS
 #define EXCLUDE_NONPAK_UE_EXTENSIONS 1	// Use .Build.cs file to disable this if the game relies on accessing loose files
 #endif
diff --git a/Engine/Source/Runtime/PakFile/Public/IPlatformFilePak.h b/Engine/Source/Runtime/PakFile/Public/IPlatformFilePak.h
--- a/Engine/Source/Runtime/PakFile/Public/IPlatformFilePak.h
+++ b/Engine/Source/Runtime/PakFile/Public/IPlatformFilePak.h
@@ -34,7 +34,14 @@ typedef FSHAHash TPakChunkHash;
 #endif
 
 PAKFILE_API TPakChunkHash ComputePakChunkHash(const void* InData, int64 InDataSizeInBytes);
-PAKFILE_API FString ChunkHashToString(const TPakChunkHash& InHash);
+FORCEINLINE FString ChunkHashToString(const TPakChunkHash& InHash)
+{
+#if PAKHASH_USE_CRC
+	return FString::Printf(TEXT("%08X"), InHash);
+#else
+	return LexToString(InHash);
+#endif
+}
 
 struct FPakChunkSignatureCheckFailedData
 {
