Improvements to loading/validation/logging of pak signature files. Better logging of errors.

diff --git a/Engine/Source/Runtime/PakFile/Private/IPlatformFilePak.cpp b/Engine/Source/Runtime/PakFile/Private/IPlatformFilePak.cpp
--- a/Engine/Source/Runtime/PakFile/Private/IPlatformFilePak.cpp
+++ b/Engine/Source/Runtime/PakFile/Private/IPlatformFilePak.cpp
@@ -4382,10 +4382,10 @@ void FPakFile::Initialize(FArchive* Reader)
 			{
 				ensure(Check());
 			}
-
-			// LoadIndex should crash in case of an error, so just assume everything is ok if we got here.
-			bIsValid = true;
 		}
+
+		// LoadIndex should crash in case of an error, so just assume everything is ok if we got here.
+		bIsValid = true;
 	}
 }
 
@@ -5570,53 +5570,52 @@ bool FPakPlatformFile::Mount(const TCHAR* InPakFilename, uint32 PakOrder, const
 		FPakFile* Pak = new FPakFile(LowerLevel, InPakFilename, bSigned);
 		if (Pak->IsValid())
 		{
-			if (InPath != NULL)
-			{
-				Pak->SetMountPoint(InPath);
-			}
-			FString PakFilename = InPakFilename;
-			if (PakFilename.EndsWith(TEXT("_P.pak")))
+			if (GetRegisteredEncryptionKeys().HasKey(Pak->GetInfo().EncryptionKeyGuid))
 			{
-				// Prioritize based on the chunk version number
-				// Default to version 1 for single patch system
-				uint32 ChunkVersionNumber = 1;
-				FString StrippedPakFilename = PakFilename.LeftChop(6);
-				int32 VersionEndIndex = PakFilename.Find("_", ESearchCase::CaseSensitive, ESearchDir::FromEnd);
-				if (VersionEndIndex != INDEX_NONE && VersionEndIndex > 0)
+				if (InPath != NULL)
+				{
+					Pak->SetMountPoint(InPath);
+				}
+				FString PakFilename = InPakFilename;
+				if (PakFilename.EndsWith(TEXT("_P.pak")))
 				{
-					int32 VersionStartIndex = PakFilename.Find("_", ESearchCase::CaseSensitive, ESearchDir::FromEnd, VersionEndIndex - 1);
-					if (VersionStartIndex != INDEX_NONE)
+					// Prioritize based on the chunk version number
+					// Default to version 1 for single patch system
+					uint32 ChunkVersionNumber = 1;
+					FString StrippedPakFilename = PakFilename.LeftChop(6);
+					int32 VersionEndIndex = PakFilename.Find("_", ESearchCase::CaseSensitive, ESearchDir::FromEnd);
+					if (VersionEndIndex != INDEX_NONE && VersionEndIndex > 0)
 					{
-						VersionStartIndex++;
-						FString VersionString = PakFilename.Mid(VersionStartIndex, VersionEndIndex - VersionStartIndex);
-						if (VersionString.IsNumeric())
+						int32 VersionStartIndex = PakFilename.Find("_", ESearchCase::CaseSensitive, ESearchDir::FromEnd, VersionEndIndex - 1);
+						if (VersionStartIndex != INDEX_NONE)
 						{
-							int32 ChunkVersionSigned = FCString::Atoi(*VersionString);
-							if (ChunkVersionSigned >= 1)
+							VersionStartIndex++;
+							FString VersionString = PakFilename.Mid(VersionStartIndex, VersionEndIndex - VersionStartIndex);
+							if (VersionString.IsNumeric())
 							{
-								// Increment by one so that the first patch file still gets more priority than the base pak file
-								ChunkVersionNumber = (uint32)ChunkVersionSigned + 1;
+								int32 ChunkVersionSigned = FCString::Atoi(*VersionString);
+								if (ChunkVersionSigned >= 1)
+								{
+									// Increment by one so that the first patch file still gets more priority than the base pak file
+									ChunkVersionNumber = (uint32)ChunkVersionSigned + 1;
+								}
 							}
 						}
 					}
+					PakOrder += 100 * ChunkVersionNumber;
 				}
-				PakOrder += 100 * ChunkVersionNumber;
-			}
-			{
-				// Add new pak file
-				FScopeLock ScopedLock(&PakListCritical);
-				FPakListEntry Entry;
-				Entry.ReadOrder = PakOrder;
-				Entry.PakFile = Pak;
-				PakFiles.Add(Entry);
-				PakFiles.StableSort();
+				{
+					// Add new pak file
+					FScopeLock ScopedLock(&PakListCritical);
+					FPakListEntry Entry;
+					Entry.ReadOrder = PakOrder;
+					Entry.PakFile = Pak;
+					PakFiles.Add(Entry);
+					PakFiles.StableSort();
+				}
+				bSuccess = true;
 			}
-			bSuccess = true;
-		}
-		else
-		{
-			bool bPassedSignatureChecks = !bSigned || Pak->PassedSignatureChecks();
-			if (bPassedSignatureChecks && Pak->GetInfo().EncryptionKeyGuid.IsValid())
+			else
 			{
 				UE_LOG(LogPakFile, Display, TEXT("Deferring mount of pak \"%s\" until encryption key '%s' becomes available"), InPakFilename, *Pak->GetInfo().EncryptionKeyGuid.ToString());
 
@@ -5632,15 +5631,15 @@ bool FPakPlatformFile::Mount(const TCHAR* InPakFilename, uint32 PakOrder, const
 				PakHandle.Reset();
 				return false;
 			}
-			else
-			{
-				UE_LOG(LogPakFile, Warning, TEXT("Failed to mount pak \"%s\", pak is invalid."), InPakFilename);
-			}
+		}
+		else
+		{
+			UE_LOG(LogPakFile, Warning, TEXT("Failed to mount pak \"%s\", pak is invalid."), InPakFilename);
 		}
 	}
 	else
 	{
-		UE_LOG(LogPakFile, Warning, TEXT("Pak \"%s\" does not exist!"), InPakFilename);
+		UE_LOG(LogPakFile, Warning, TEXT("Failed to open pak \"%s\""), InPakFilename);
 	}
 	return bSuccess;
 }
diff --git a/Engine/Source/Runtime/PakFile/Public/IPlatformFilePak.h b/Engine/Source/Runtime/PakFile/Public/IPlatformFilePak.h
--- a/Engine/Source/Runtime/PakFile/Public/IPlatformFilePak.h
+++ b/Engine/Source/Runtime/PakFile/Public/IPlatformFilePak.h
@@ -2389,7 +2389,7 @@ struct FPakSignatureFile
 
 	enum class EVersion
 	{
-		Legacy,
+		Invalid,
 		First,
 
 		Last,
@@ -2423,20 +2423,14 @@ struct FPakSignatureFile
 	 */
 	void Serialize(FArchive& Ar)
 	{
-		int64 StartPos = Ar.Tell();
 		uint32 FileMagic = Magic;
 		Ar << FileMagic;
 
 		if (Ar.IsLoading() && FileMagic != Magic)
 		{
-			// Old format with no versioning! Go back to where we were and mark our version as legacy
-			Ar.Seek(StartPos);
-			Version = EVersion::Legacy;
-			TEncryptionInt LegacyEncryptedCRC;
-			Ar << LegacyEncryptedCRC;
-			Ar << ChunkHashes;
-			// Note that we don't do any actual signature checking here because this is old data that won't have been
-			// encrypted with the new larger keys.
+			Version = EVersion::Invalid;
+			EncryptedHash.Empty();
+			ChunkHashes.Empty();
 			return;
 		}
 
@@ -2459,17 +2453,27 @@ struct FPakSignatureFile
 	 */
 	bool DecryptSignatureAndValidate(FRSA::TKeyPtr InKey, const FString& InFilename)
 	{
-		FRSA::Decrypt(FRSA::EKeyType::Public, EncryptedHash, DecryptedHash.Hash, ARRAY_COUNT(FSHAHash::Hash), InKey);
-
-		FSHAHash CurrentHash = ComputeCurrentMasterHash();
-		if (DecryptedHash != CurrentHash)
+		if (Version == EVersion::Invalid)
+		{
+			UE_LOG(LogPakFile, Warning, TEXT("Pak signature file for '%s' was invalid"), *InFilename);
+		}
+		else if (FRSA::Decrypt(FRSA::EKeyType::Public, EncryptedHash, DecryptedHash.Hash, ARRAY_COUNT(FSHAHash::Hash), InKey))
 		{
+			FSHAHash CurrentHash = ComputeCurrentMasterHash();
+			if (DecryptedHash == CurrentHash)
+			{
+				return true;
+			}
+
 			UE_LOG(LogPakFile, Warning, TEXT("Pak signature table validation failed for '%s'! Expected %s, Received %s"), *InFilename, *DecryptedHash.ToString(), *CurrentHash.ToString());
-			FPakPlatformFile::GetPakMasterSignatureTableCheckFailureHandler().Broadcast(InFilename);
-			return false;
+		}
+		else
+		{
+			UE_LOG(LogPakFile, Warning, TEXT("Pak signature table validation failed for '%s'! Failed to decrypt signature"), *InFilename);
 		}
 
-		return true;
+		FPakPlatformFile::GetPakMasterSignatureTableCheckFailureHandler().Broadcast(InFilename);
+		return false;
 	}
 
 	/**
