Minor off-by-one correction to chunk hash failure warning message

diff --git a/Engine/Source/Runtime/PakFile/Private/IPlatformFilePak.cpp b/Engine/Source/Runtime/PakFile/Private/IPlatformFilePak.cpp
--- a/Engine/Source/Runtime/PakFile/Private/IPlatformFilePak.cpp
+++ b/Engine/Source/Runtime/PakFile/Private/IPlatformFilePak.cpp
@@ -3225,7 +3225,7 @@ void FPakPrecacher::DoSignatureCheck(bool bWasCanceled, IAsyncReadRequest* Reque
 				FScopeLock Lock(&CachedFilesScopeLock);
 				FPakData* PakData = &CachedPakData[PakIndex];
 
-				UE_LOG(LogPakFile, Warning, TEXT("Pak chunk signing mismatch on chunk [%i/%i]! Expected %s, Received %s"), SignatureIndex, PakData->Signatures->ChunkHashes.Num(), *ChunkHashToString(PakData->Signatures->ChunkHashes[SignatureIndex]), *ChunkHashToString(ThisHash));
+				UE_LOG(LogPakFile, Warning, TEXT("Pak chunk signing mismatch on chunk [%i/%i]! Expected %s, Received %s"), SignatureIndex, PakData->Signatures->ChunkHashes.Num() - 1, *ChunkHashToString(PakData->Signatures->ChunkHashes[SignatureIndex]), *ChunkHashToString(ThisHash));
 
 				// Check the signatures are still as we expected them
 				if (PakData->Signatures->DecryptedHash != PakData->Signatures->ComputeCurrentMasterHash())
diff --git a/Engine/Source/Runtime/PakFile/Private/SignedArchiveReader.cpp b/Engine/Source/Runtime/PakFile/Private/SignedArchiveReader.cpp
--- a/Engine/Source/Runtime/PakFile/Private/SignedArchiveReader.cpp
+++ b/Engine/Source/Runtime/PakFile/Private/SignedArchiveReader.cpp
@@ -243,7 +243,7 @@ bool FChunkCacheWorker::CheckSignature(const FChunkRequest& ChunkInfo)
 		bChunkHashesMatch = IsValid() && (ChunkHash == Signatures->ChunkHashes[ChunkInfo.Index]);
 		if (!bChunkHashesMatch)
 		{
-			UE_LOG(LogPakFile, Warning, TEXT("Pak chunk signing mismatch on chunk [%i/%i]! Expected %s, Received %s"), ChunkInfo.Index, Signatures->ChunkHashes.Num(), *ChunkHashToString(Signatures->ChunkHashes[ChunkInfo.Index]), *ChunkHashToString(ChunkHash));
+			UE_LOG(LogPakFile, Warning, TEXT("Pak chunk signing mismatch on chunk [%i/%i]! Expected %s, Received %s"), ChunkInfo.Index, Signatures->ChunkHashes.Num() - 1, *ChunkHashToString(Signatures->ChunkHashes[ChunkInfo.Index]), *ChunkHashToString(ChunkHash));
 
 			if (Signatures->DecryptedHash != Signatures->ComputeCurrentMasterHash())
 			{
