Improvements to pak file check command
 - Add timing information to the log
 - If pak is signed, use a faster method where we just read one byte from each signing block
 - Still invoked with the -checkpak command line parameter

diff --git a/Engine/Source/Runtime/PakFile/Private/IPlatformFilePak.cpp b/Engine/Source/Runtime/PakFile/Private/IPlatformFilePak.cpp
--- a/Engine/Source/Runtime/PakFile/Private/IPlatformFilePak.cpp
+++ b/Engine/Source/Runtime/PakFile/Private/IPlatformFilePak.cpp
@@ -4496,54 +4496,87 @@ void FPakFile::LoadIndex(FArchive* Reader)
 bool FPakFile::Check()
 {
 	UE_LOG(LogPakFile, Display, TEXT("Checking pak file \"%s\". This may take a while..."), *PakFilename);
+	double StartTime = FPlatformTime::Seconds();
+
 	FArchive& PakReader = *GetSharedReader(NULL);
 	int32 ErrorCount = 0;
 	int32 FileCount = 0;
 
-	const bool bIncludeDeleted = true;
-	for (FPakFile::FFileIterator It(*this,bIncludeDeleted); It; ++It, ++FileCount)
+	bool bSuccess = true;
+
+	// If the pak file is signed, we can do a fast check by just reading a single byte from the start of
+	// each signing block. The signed archive reader will bring in that whole signing block and compare
+	// against the signature table and fire the handler
+	if (bSigned)
 	{
-		const FPakEntry& Entry = It.Info();
-		if( Entry.IsDeleteRecord() )
+		FDelegateHandle DelegateHandle = FPakPlatformFile::GetPakChunkSignatureCheckFailedHandler().AddLambda([&bSuccess](const FPakChunkSignatureCheckFailedData&)
+		{ 
+			bSuccess = false; 
+		});
+
+		int64 CurrentPos = 0;
+		const int64 Size = PakReader.TotalSize();
+		while (CurrentPos < Size)
 		{
-			UE_LOG(LogPakFile, Display, TEXT("\"%s\" Deleted."), *It.Filename());
-			continue;
+			PakReader.Seek(CurrentPos);
+			uint8 Byte = 0;
+			PakReader << Byte;
+			CurrentPos += FPakInfo::MaxChunkDataSize;
 		}
 
-		void* FileContents = FMemory::Malloc(Entry.Size);
-		PakReader.Seek(Entry.Offset);
-		uint32 SerializedCrcTest = 0;
-		FPakEntry EntryInfo;
-		EntryInfo.Serialize(PakReader, GetInfo().Version);
-		if (EntryInfo != Entry)
+		FPakPlatformFile::GetPakChunkSignatureCheckFailedHandler().Remove(DelegateHandle);
+	}
+	else
+	{
+		const bool bIncludeDeleted = true;
+		for (FPakFile::FFileIterator It(*this, bIncludeDeleted); It; ++It, ++FileCount)
 		{
-			UE_LOG(LogPakFile, Error, TEXT("Serialized hash mismatch for \"%s\"."), *It.Filename());
-			ErrorCount++;
-		}
-		PakReader.Serialize(FileContents, Entry.Size);
+			const FPakEntry& Entry = It.Info();
+			if (Entry.IsDeleteRecord())
+			{
+				UE_LOG(LogPakFile, Verbose, TEXT("\"%s\" Deleted."), *It.Filename());
+				continue;
+			}
 
-		uint8 TestHash[20];
-		FSHA1::HashBuffer(FileContents, Entry.Size, TestHash);
-		if (FMemory::Memcmp(TestHash, Entry.Hash, sizeof(TestHash)) != 0)
+			void* FileContents = FMemory::Malloc(Entry.Size);
+			PakReader.Seek(Entry.Offset);
+			uint32 SerializedCrcTest = 0;
+			FPakEntry EntryInfo;
+			EntryInfo.Serialize(PakReader, GetInfo().Version);
+			if (EntryInfo != Entry)
+			{
+				UE_LOG(LogPakFile, Error, TEXT("Serialized hash mismatch for \"%s\"."), *It.Filename());
+				ErrorCount++;
+			}
+			PakReader.Serialize(FileContents, Entry.Size);
+
+			uint8 TestHash[20];
+			FSHA1::HashBuffer(FileContents, Entry.Size, TestHash);
+			if (FMemory::Memcmp(TestHash, Entry.Hash, sizeof(TestHash)) != 0)
+			{
+				UE_LOG(LogPakFile, Error, TEXT("Hash mismatch for \"%s\"."), *It.Filename());
+				ErrorCount++;
+			}
+			else
+			{
+				UE_LOG(LogPakFile, Verbose, TEXT("\"%s\" OK. [%s]"), *It.Filename(), *Info.GetCompressionMethod(Entry.CompressionMethodIndex).ToString());
+			}
+			FMemory::Free(FileContents);
+		}
+		if (ErrorCount == 0)
 		{
-			UE_LOG(LogPakFile, Error, TEXT("Hash mismatch for \"%s\"."), *It.Filename());
-			ErrorCount++;
+			UE_LOG(LogPakFile, Display, TEXT("Pak file \"%s\" healthy, %d files checked."), *PakFilename, FileCount);
 		}
 		else
 		{
-			UE_LOG(LogPakFile, Display, TEXT("\"%s\" OK. [%s]"), *It.Filename(), *Info.GetCompressionMethod(Entry.CompressionMethodIndex).ToString());
+			UE_LOG(LogPakFile, Display, TEXT("Pak file \"%s\" corrupted (%d errors out of %d files checked.)."), *PakFilename, ErrorCount, FileCount);
 		}
-		FMemory::Free(FileContents);
-	}
-	if (ErrorCount == 0)
-	{
-		UE_LOG(LogPakFile, Display, TEXT("Pak file \"%s\" healthy, %d files checked."), *PakFilename, FileCount);
-	}
-	else
-	{
-		UE_LOG(LogPakFile, Display, TEXT("Pak file \"%s\" corrupted (%d errors out of %d files checked.)."), *PakFilename, ErrorCount, FileCount);
 	}
 
+	double EndTime = FPlatformTime::Seconds();
+	double ElapsedTime = EndTime - StartTime;
+	UE_LOG(LogPakFile, Display, TEXT("Pak file \"%s\" checked in %.2fs"), *PakFilename, ElapsedTime);
+
 	return ErrorCount == 0;
 }
 
