Add support to pak platform file for a custom decryption function that a game can override the way decryption is performed at runtime

diff --git a/Engine/Source/Runtime/PakFile/Private/IPlatformFilePak.cpp b/Engine/Source/Runtime/PakFile/Private/IPlatformFilePak.cpp
--- a/Engine/Source/Runtime/PakFile/Private/IPlatformFilePak.cpp
+++ b/Engine/Source/Runtime/PakFile/Private/IPlatformFilePak.cpp
@@ -168,6 +168,12 @@ FFilenameSecurityDelegate& FPakPlatformFile::GetFilenameSecurityDelegate()
 	return Delegate;
 }
 
+FPakCustomEncryptionDelegate& FPakPlatformFile::GetPakCustomEncryptionDelegate()
+{
+	static FPakCustomEncryptionDelegate Delegate;
+	return Delegate;
+}
+
 FPakChunkSignatureCheckFailedHandler& FPakPlatformFile::GetPakChunkSignatureCheckFailedHandler()
 {
 	static FPakChunkSignatureCheckFailedHandler Delegate;
@@ -288,11 +294,18 @@ DECLARE_DWORD_ACCUMULATOR_STAT(TEXT("PakCache Async Decrypts (Uncompressed Path)
 
 void DecryptData(uint8* InData, uint32 InDataSize, FGuid InEncryptionKeyGuid)
 {
-	SCOPE_SECONDS_ACCUMULATOR(STAT_PakCache_DecryptTime);
-	FAES::FAESKey Key;
-	FPakPlatformFile::GetPakEncryptionKey(Key, InEncryptionKeyGuid);
-	check(Key.IsValid());
-	FAES::DecryptData(InData, InDataSize, Key);
+	if (FPakPlatformFile::GetPakCustomEncryptionDelegate().IsBound())
+	{
+		FPakPlatformFile::GetPakCustomEncryptionDelegate().Execute(InData, InDataSize, InEncryptionKeyGuid);
+	}
+	else
+	{
+		SCOPE_SECONDS_ACCUMULATOR(STAT_PakCache_DecryptTime);
+		FAES::FAESKey Key;
+		FPakPlatformFile::GetPakEncryptionKey(Key, InEncryptionKeyGuid);
+		check(Key.IsValid());
+		FAES::DecryptData(InData, InDataSize, Key);
+	}
 }
 
 #if USE_PAK_PRECACHE
diff --git a/Engine/Source/Runtime/PakFile/Public/IPlatformFilePak.h b/Engine/Source/Runtime/PakFile/Public/IPlatformFilePak.h
--- a/Engine/Source/Runtime/PakFile/Public/IPlatformFilePak.h
+++ b/Engine/Source/Runtime/PakFile/Public/IPlatformFilePak.h
@@ -61,6 +61,7 @@ struct FPakChunkSignatureCheckFailedData
 };
 /** Delegate for allowing a game to restrict the accessing of non-pak files */
 DECLARE_DELEGATE_RetVal_OneParam(bool, FFilenameSecurityDelegate, const TCHAR* /*InFilename*/);
+DECLARE_DELEGATE_ThreeParams(FPakCustomEncryptionDelegate, uint8* /*InData*/, uint32 /*InDataSize*/, FGuid /*InEncryptionKeyGuid*/);
 DECLARE_MULTICAST_DELEGATE_OneParam(FPakChunkSignatureCheckFailedHandler, const FPakChunkSignatureCheckFailedData&);
 DECLARE_MULTICAST_DELEGATE_OneParam(FPakMasterSignatureTableCheckFailureHandler, const FString&);
 
@@ -2347,6 +2348,9 @@ public:
 	// Access static delegate for loose file security
 	static FFilenameSecurityDelegate& GetFilenameSecurityDelegate();
 
+	// Access static delegate for custom encryption
+	static FPakCustomEncryptionDelegate& GetPakCustomEncryptionDelegate();
+
 	// Access static delegate for handling a pak signature check failure
 	static FPakChunkSignatureCheckFailedHandler& GetPakChunkSignatureCheckFailedHandler();
 
