Bundle Manager now makes calls to optimize pak memory usage after mounting.

diff --git a/Engine/Source/Runtime/Core/Private/Misc/CoreDelegates.cpp b/Engine/Source/Runtime/Core/Private/Misc/CoreDelegates.cpp
--- a/Engine/Source/Runtime/Core/Private/Misc/CoreDelegates.cpp
+++ b/Engine/Source/Runtime/Core/Private/Misc/CoreDelegates.cpp
@@ -29,6 +29,8 @@ FCoreDelegates::FOnPreMainInit& FCoreDelegates::GetPreMainInitDelegate()
 FCoreDelegates::FOnMountAllPakFiles FCoreDelegates::OnMountAllPakFiles;
 FCoreDelegates::FOnMountPak FCoreDelegates::OnMountPak;
 FCoreDelegates::FOnUnmountPak FCoreDelegates::OnUnmountPak;
+FCoreDelegates::FOnOptimizeMemoryUsageForMountedPaks FCoreDelegates::OnOptimizeMemoryUsageForMountedPaks;
+
 FCoreDelegates::FPakFileMountedDelegate FCoreDelegates::PakFileMountedCallback;
 FCoreDelegates::FNoPakFilesMountedDelegate FCoreDelegates::NoPakFilesMountedDelegate;
 FCoreDelegates::FOnFileOpenedForReadFromPakFile FCoreDelegates::OnFileOpenedForReadFromPakFile;
diff --git a/Engine/Source/Runtime/Core/Public/Misc/CoreDelegates.h b/Engine/Source/Runtime/Core/Public/Misc/CoreDelegates.h
--- a/Engine/Source/Runtime/Core/Public/Misc/CoreDelegates.h
+++ b/Engine/Source/Runtime/Core/Public/Misc/CoreDelegates.h
@@ -64,6 +64,9 @@ public:
 	// delegate type for prompting the pak system to mount a new pak
 	DECLARE_DELEGATE_RetVal_OneParam(bool, FOnUnmountPak, const FString&);
 
+	// delegate type for prompting the pak system to optimize memory for mounted paks
+	DECLARE_DELEGATE(FOnOptimizeMemoryUsageForMountedPaks);
+
 	// delegate for handling when a new pak file is successfully mounted passes in the name of the mounted pak file
 	DECLARE_MULTICAST_DELEGATE_OneParam(FPakFileMountedDelegate, const TCHAR*);
 
@@ -152,6 +155,9 @@ public:
 	// Callback for unmounting a pak file.
 	static FOnUnmountPak OnUnmountPak;
 
+	// Callback to optimize memeory for currently mounted paks
+	static FOnOptimizeMemoryUsageForMountedPaks OnOptimizeMemoryUsageForMountedPaks;
+
 	// After a pakfile is mounted this callback is called for each new file
 	static FPakFileMountedDelegate PakFileMountedCallback;
 
diff --git a/Engine/Source/Runtime/PakFile/Private/IPlatformFilePak.cpp b/Engine/Source/Runtime/PakFile/Private/IPlatformFilePak.cpp
--- a/Engine/Source/Runtime/PakFile/Private/IPlatformFilePak.cpp
+++ b/Engine/Source/Runtime/PakFile/Private/IPlatformFilePak.cpp
@@ -5402,8 +5402,13 @@ FPakPlatformFile::FPakPlatformFile()
 FPakPlatformFile::~FPakPlatformFile()
 {
 	FCoreDelegates::GetRegisterEncryptionKeyDelegate().Unbind();
+
+	FCoreDelegates::OnFEngineLoopInitComplete.RemoveAll(this);
+
+	FCoreDelegates::OnMountAllPakFiles.Unbind();
 	FCoreDelegates::OnMountPak.Unbind();
 	FCoreDelegates::OnUnmountPak.Unbind();
+	FCoreDelegates::OnOptimizeMemoryUsageForMountedPaks.Unbind();
 
 #if USE_PAK_PRECACHE
 	FPakPrecacher::Shutdown();
@@ -5554,14 +5559,38 @@ bool FPakPlatformFile::Initialize(IPlatformFile* Inner, const TCHAR* CmdLine)
 	FCoreDelegates::OnMountAllPakFiles.BindRaw(this, &FPakPlatformFile::MountAllPakFiles);
 	FCoreDelegates::OnMountPak.BindRaw(this, &FPakPlatformFile::HandleMountPakDelegate);
 	FCoreDelegates::OnUnmountPak.BindRaw(this, &FPakPlatformFile::HandleUnmountPakDelegate);
+	FCoreDelegates::OnOptimizeMemoryUsageForMountedPaks.BindRaw(this, &FPakPlatformFile::OptimizeMemoryUsageForMountedPaks);
+
+	FCoreDelegates::OnFEngineLoopInitComplete.AddRaw(this, &FPakPlatformFile::OptimizeMemoryUsageForMountedPaks);
+
+	return !!LowerLevel;
+}
+
+void FPakPlatformFile::InitializeNewAsyncIO()
+{
+#if USE_PAK_PRECACHE
+#if !WITH_EDITOR
+	if (FPlatformProcess::SupportsMultithreading() && !FParse::Param(FCommandLine::Get(), TEXT("FileOpenLog")))
+	{
+		FPakPrecacher::Init(LowerLevel, GetPakSigningKey());
+	}
+	else
+#endif
+	{
+		UE_CLOG(FParse::Param(FCommandLine::Get(), TEXT("FileOpenLog")), LogPakFile, Display, TEXT("Disabled pak precacher to get an accurate load order. This should only be used to collect gameopenorder.log, as it is quite slow."));
+		GPakCache_Enable = 0;
+	}
+#endif
+}
 
+void FPakPlatformFile::OptimizeMemoryUsageForMountedPaks()
+{
 #if !(IS_PROGRAM || WITH_EDITOR)
-	FCoreDelegates::OnFEngineLoopInitComplete.AddLambda([this] {
 			FPlatformMisc::LowLevelOutputDebugStringf(TEXT("Checking Pak Config\n"));
 			bool bUnloadPakEntryFilenamesIfPossible = FParse::Param(FCommandLine::Get(), TEXT("unloadpakentryfilenames"));
 			GConfig->GetBool(TEXT("Pak"), TEXT("UnloadPakEntryFilenamesIfPossible"), bUnloadPakEntryFilenamesIfPossible, GEngineIni);
 
-        if ((bUnloadPakEntryFilenamesIfPossible && !FParse::Param(FCommandLine::Get(), TEXT("nounloadpakentries"))) || FParse::Param(FCommandLine::Get(),TEXT("unloadpakentries")))
+	if ((bUnloadPakEntryFilenamesIfPossible && !FParse::Param(FCommandLine::Get(), TEXT("nounloadpakentries"))) || FParse::Param(FCommandLine::Get(), TEXT("unloadpakentries")))
 			{
 				// With [Pak] UnloadPakEntryFilenamesIfPossible enabled, [Pak] DirectoryRootsToKeepInMemoryWhenUnloadingPakEntryFilenames
 				// can contain pak entry directory wildcards of which the entire recursive directory structure of filenames underneath a
@@ -5585,26 +5614,6 @@ bool FPakPlatformFile::Initialize(IPlatformFile* Inner, const TCHAR* CmdLine)
 				FPakPlatformFile* PakPlatformFile = (FPakPlatformFile*)(FPlatformFileManager::Get().FindPlatformFile(FPakPlatformFile::GetTypeName()));
 				PakPlatformFile->ShrinkPakEntriesMemoryUsage();
 			}
-		});
-#endif
-
-	return !!LowerLevel;
-}
-
-void FPakPlatformFile::InitializeNewAsyncIO()
-{
-#if USE_PAK_PRECACHE
-#if !WITH_EDITOR
-	if (FPlatformProcess::SupportsMultithreading() && !FParse::Param(FCommandLine::Get(), TEXT("FileOpenLog")))
-	{
-		FPakPrecacher::Init(LowerLevel, GetPakSigningKey());
-	}
-	else
-#endif
-	{
-		UE_CLOG(FParse::Param(FCommandLine::Get(), TEXT("FileOpenLog")), LogPakFile, Display, TEXT("Disabled pak precacher to get an accurate load order. This should only be used to collect gameopenorder.log, as it is quite slow."));
-		GPakCache_Enable = 0;
-	}
 #endif
 }
 
diff --git a/Engine/Source/Runtime/PakFile/Public/IPlatformFilePak.h b/Engine/Source/Runtime/PakFile/Public/IPlatformFilePak.h
--- a/Engine/Source/Runtime/PakFile/Public/IPlatformFilePak.h
+++ b/Engine/Source/Runtime/PakFile/Public/IPlatformFilePak.h
@@ -1581,6 +1581,8 @@ public:
 	virtual bool Initialize(IPlatformFile* Inner, const TCHAR* CommandLineParam) override;
 	virtual void InitializeNewAsyncIO() override;
 
+	void OptimizeMemoryUsageForMountedPaks();
+
 	virtual IPlatformFile* GetLowerLevel() override
 	{
 		return LowerLevel;
