Hardening of pak index loading
 - If a corrupt index is detected, do a retry but collect and log more information along the way to help diagnose the issue.

diff --git a/Engine/Source/Developer/PakFileUtilities/Private/PakFileUtilities.cpp b/Engine/Source/Developer/PakFileUtilities/Private/PakFileUtilities.cpp
--- a/Engine/Source/Developer/PakFileUtilities/Private/PakFileUtilities.cpp
+++ b/Engine/Source/Developer/PakFileUtilities/Private/PakFileUtilities.cpp
@@ -1998,7 +1998,7 @@ bool CreatePakFile(const TCHAR* Filename, TArray<FPakInputPair>& FilesToAdd, con
 		}
 	}
 
-	FSHA1::HashBuffer(IndexData.GetData(), IndexData.Num(), Info.IndexHash);
+	FSHA1::HashBuffer(IndexData.GetData(), IndexData.Num(), Info.IndexHash.Hash);
 
 	if (Info.bEncryptedIndex)
 	{
diff --git a/Engine/Source/Runtime/PakFile/Private/IPlatformFilePak.cpp b/Engine/Source/Runtime/PakFile/Private/IPlatformFilePak.cpp
--- a/Engine/Source/Runtime/PakFile/Private/IPlatformFilePak.cpp
+++ b/Engine/Source/Runtime/PakFile/Private/IPlatformFilePak.cpp
@@ -4411,44 +4411,85 @@ void FPakFile::LoadIndex(FArchive* Reader)
 	else
 	{
 		// Load index into memory first.
-		Reader->Seek(Info.IndexOffset);
+
+		// If we encounter a corrupt index, try again but gather some extra debugging information so we can try and understand where it failed.
+		bool bFirstPass = true;
 		TArray<uint8> IndexData;
-		IndexData.AddUninitialized(Info.IndexSize);
-		Reader->Serialize(IndexData.GetData(), Info.IndexSize);
-		FMemoryReader IndexReader(IndexData);
 
-		// Decrypt if necessary
-		if (Info.bEncryptedIndex)
+		while (true)
 		{
-			DecryptData(IndexData.GetData(), Info.IndexSize, Info.EncryptionKeyGuid);
-		}
+			Reader->Seek(Info.IndexOffset);
+			IndexData.Empty(); // The next SetNum makes this Empty() logically redundant, but we want to try and force a memory reallocation for the re-attempt
+			IndexData.SetNum(Info.IndexSize);
+			Reader->Serialize(IndexData.GetData(), Info.IndexSize);
 
-		// Check SHA1 value.
-		uint8 IndexHash[20];
-		FSHA1::HashBuffer(IndexData.GetData(), IndexData.Num(), IndexHash);
-		if (FMemory::Memcmp(IndexHash, Info.IndexHash, sizeof(IndexHash)) != 0)
-		{
-			FString StoredIndexHash, ComputedIndexHash;
-			StoredIndexHash = TEXT("0x");
-			ComputedIndexHash = TEXT("0x");
+			FSHAHash EncryptedDataHash;
+			if (!bFirstPass)
+			{
+				FSHA1::HashBuffer(IndexData.GetData(), IndexData.Num(), EncryptedDataHash.Hash);
+			}
 
-			for (int64 ByteIndex = 0; ByteIndex < 20; ++ByteIndex)
+			// Decrypt if necessary
+			if (Info.bEncryptedIndex)
 			{
-				StoredIndexHash += FString::Printf(TEXT("%02X"), Info.IndexHash[ByteIndex]);
-				ComputedIndexHash += FString::Printf(TEXT("%02X"), IndexHash[ByteIndex]);
+				DecryptData(IndexData.GetData(), Info.IndexSize, Info.EncryptionKeyGuid);
 			}
+			// Check SHA1 value.
+			FSHAHash ComputedHash;
+			FSHA1::HashBuffer(IndexData.GetData(), IndexData.Num(), ComputedHash.Hash);
+			if (Info.IndexHash != ComputedHash)
+			{
+				if (bFirstPass)
+				{
+					UE_LOG(LogPakFile, Log, TEXT("Corrupt pak index detected!"));
+					UE_LOG(LogPakFile, Log, TEXT(" Filename: %s"), *PakFilename);
+					UE_LOG(LogPakFile, Log, TEXT(" Encrypted: %d"), Info.bEncryptedIndex);
+					UE_LOG(LogPakFile, Log, TEXT(" Total Size: %d"), Reader->TotalSize());
+					UE_LOG(LogPakFile, Log, TEXT(" Index Offset: %d"), Info.IndexOffset);
+					UE_LOG(LogPakFile, Log, TEXT(" Index Size: %d"), Info.IndexSize);
+					UE_LOG(LogPakFile, Log, TEXT(" Stored Index Hash: %s"), *Info.IndexHash.ToString());
+					UE_LOG(LogPakFile, Log, TEXT(" Computed Index Hash [Pass 0]: %s"), *ComputedHash.ToString());
+					bFirstPass = false;
+				}
+				else
+				{
+					UE_LOG(LogPakFile, Log, TEXT(" Computed Index Hash [Pass 1]: %s"), *ComputedHash.ToString());
+					UE_LOG(LogPakFile, Log, TEXT(" Encrypted Index Hash: %s"), *EncryptedDataHash.ToString());
+
+					// Compute an SHA1 hash of the whole file so we can tell if the file was modified on disc (obviously as long as this isn't an  IO bug that gives us the same bogus data again)
+					FSHA1 FileHash;
+					Reader->Seek(0);
+					int64 Remaining = Reader->TotalSize();
+					TArray<uint8> WorkingBuffer;
+					WorkingBuffer.SetNum(64 * 1024);
+					while (Remaining > 0)
+					{
+						int64 ToProcess = FMath::Min((int64)WorkingBuffer.Num(), Remaining);
+						Reader->Serialize(WorkingBuffer.GetData(), ToProcess);
+						FileHash.Update(WorkingBuffer.GetData(), ToProcess);
+						Remaining -= ToProcess;
+					}
+
+					FileHash.Final();
+					FSHAHash FinalFileHash;
+					FileHash.GetHash(FinalFileHash.Hash);
+					UE_LOG(LogPakFile, Log, TEXT(" File Hash: %s"), *FinalFileHash.ToString());
 
-			UE_LOG(LogPakFile, Log, TEXT("Corrupt pak index detected!"));
-			UE_LOG(LogPakFile, Log, TEXT(" Filename: %s"), *PakFilename);
-			UE_LOG(LogPakFile, Log, TEXT(" Encrypted: %d"), Info.bEncryptedIndex);
-			UE_LOG(LogPakFile, Log, TEXT(" Total Size: %d"), Reader->TotalSize());
-			UE_LOG(LogPakFile, Log, TEXT(" Index Offset: %d"), Info.IndexOffset);
-			UE_LOG(LogPakFile, Log, TEXT(" Index Size: %d"), Info.IndexSize);
-			UE_LOG(LogPakFile, Log, TEXT(" Stored Index Hash: %s"), *StoredIndexHash);
-			UE_LOG(LogPakFile, Log, TEXT(" Computed Index Hash: %s"), *ComputedIndexHash);
-			UE_LOG(LogPakFile, Fatal, TEXT("Corrupted index in pak file (CRC mismatch)."));
+					UE_LOG(LogPakFile, Fatal, TEXT("Corrupted index in pak file (SHA hash mismatch)."));
+				}
+			}
+			else
+			{
+				if (!bFirstPass)
+				{
+					UE_LOG(LogPakFile, Log, TEXT("Pak index corruption appears to have recovered on the second attempt!"));
+				}
+				break;
+			}
 		}
 
+		FMemoryReader IndexReader(IndexData);
+
 		// Read the default mount point and all entries.
 		NumEntries = 0;
 		IndexReader << MountPoint;
diff --git a/Engine/Source/Runtime/PakFile/Public/IPlatformFilePak.h b/Engine/Source/Runtime/PakFile/Public/IPlatformFilePak.h
--- a/Engine/Source/Runtime/PakFile/Public/IPlatformFilePak.h
+++ b/Engine/Source/Runtime/PakFile/Public/IPlatformFilePak.h
@@ -109,7 +109,7 @@ struct FPakInfo
 	/** Size (in bytes) of pak file index. */
 	int64 IndexSize;
 	/** Index SHA1 value. */
-	uint8 IndexHash[20];
+	FSHAHash IndexHash;
 	/** Flag indicating if the pak index has been encrypted. */
 	uint8 bEncryptedIndex;
 	/** Encryption key guid. Empty if we should use the embedded key. */
@@ -127,7 +127,6 @@ struct FPakInfo
 		, IndexSize(0)
 		, bEncryptedIndex(0)
 	{
-		FMemory::Memset(IndexHash, 0, sizeof(IndexHash));
 		// we always put in a NAME_None entry as index 0, so that an uncompressed PakEntry will have CompressionMethodIndex of 0 and can early out easily
 		CompressionMethods.Add(NAME_None);
 	}
@@ -182,7 +181,7 @@ struct FPakInfo
 		Ar << Version;
 		Ar << IndexOffset;
 		Ar << IndexSize;
-		Ar.Serialize(IndexHash, sizeof(IndexHash));
+		Ar << IndexHash;
 
 		if (Ar.IsLoading())
 		{
