Only log pak key registration when something was actually mounted

diff --git a/Engine/Source/Runtime/PakFile/Private/IPlatformFilePak.cpp b/Engine/Source/Runtime/PakFile/Private/IPlatformFilePak.cpp
--- a/Engine/Source/Runtime/PakFile/Private/IPlatformFilePak.cpp
+++ b/Engine/Source/Runtime/PakFile/Private/IPlatformFilePak.cpp
@@ -6024,9 +6024,9 @@ void FPakPlatformFile::RegisterEncryptionKey(const FGuid& InGuid, const FAES::FA
 		}
 
 		PendingEncryptedPakFiles.RemoveAll([InGuid](const FPakListDeferredEntry& Entry) { return Entry.EncryptionKeyGuid == InGuid; });
-	}
 
-	UE_CLOG(InGuid.IsValid(), LogPakFile, Log, TEXT("Registered encryption key '%s': %d pak files mounted, %d remain pending"), *InGuid.ToString(), NumMounted, PendingEncryptedPakFiles.Num());
+		UE_CLOG(InGuid.IsValid(), LogPakFile, Log, TEXT("Registered encryption key '%s': %d pak files mounted, %d remain pending"), *InGuid.ToString(), NumMounted, PendingEncryptedPakFiles.Num());
+	}
 }
 
 IFileHandle* FPakPlatformFile::OpenRead(const TCHAR* Filename, bool bAllowWrite)
