UE4 - Fix rare race condition with cancelled IO requests in the pak precacher.

diff --git a/Engine/Source/Runtime/PakFile/Private/IPlatformFilePak.cpp b/Engine/Source/Runtime/PakFile/Private/IPlatformFilePak.cpp
--- a/Engine/Source/Runtime/PakFile/Private/IPlatformFilePak.cpp
+++ b/Engine/Source/Runtime/PakFile/Private/IPlatformFilePak.cpp
@@ -3001,10 +3001,7 @@ public:
 	virtual ~FPakProcessedReadRequest()
 	{
 		check(!MyCanceledBlocks.Num());
-		if (!bHasCancelled)
-		{
-			DoneWithRawRequests();
-		}
+		DoneWithRawRequests();
 		if (Memory && !bUserSuppliedMemory)
 		{
 			// this can happen with a race on cancel, it is ok, they didn't take the memory, free it now
@@ -3015,11 +3012,6 @@ public:
 		Memory = nullptr;
 	}
 
-	bool WasCanceled()
-	{
-		return bHasCancelled;
-	}
-
 	virtual void WaitCompletionImpl(float TimeLimitSeconds) override
 	{
 		{
@@ -3413,8 +3405,8 @@ public:
 					bAnyUnfinished = true;
 				}
 			}
-			check(!LiveRequests.Contains(Result))
-				LiveRequests.Add(Result);
+			check(!LiveRequests.Contains(Result));
+			LiveRequests.Add(Result);
 			if (!bAnyUnfinished)
 			{
 				Result->RequestIsComplete();
@@ -3575,6 +3567,7 @@ public:
 	{
 		check(!Block.RawRequest);
 		Block.RawRequest = nullptr;
+		//check(!Block.CPUWorkGraphEvent || Block.CPUWorkGraphEvent->IsComplete());
 		Block.CPUWorkGraphEvent = nullptr;
 		if (Block.Raw)
 		{
@@ -3599,9 +3592,15 @@ public:
 		Block.bInFlight = false;
 	}
 
-	void RemoveRequest(FPakProcessedReadRequest* Req, int64 Offset, int64 BytesToRead)
+	void RemoveRequest(FPakProcessedReadRequest* Req, int64 Offset, int64 BytesToRead, const bool& bAlreadyCancelled)
 	{
 		FScopeLock ScopedLock(&CriticalSection);
+		if (bAlreadyCancelled)
+		{
+			check(!LiveRequests.Contains(Req));
+			return;
+		}
+
 		check(LiveRequests.Contains(Req));
 		LiveRequests.Remove(Req);
 		int32 FirstBlock = Offset / FileEntry.CompressionBlockSize;
@@ -3627,9 +3626,11 @@ public:
 		}
 	}
 
-	void HandleCanceledRequest(TSet<FCachedAsyncBlock*>& MyCanceledBlocks, FPakProcessedReadRequest* Req, int64 Offset, int64 BytesToRead)
+	void HandleCanceledRequest(TSet<FCachedAsyncBlock*>& MyCanceledBlocks, FPakProcessedReadRequest* Req, int64 Offset, int64 BytesToRead, bool& bHasCancelledRef)
 	{
 		FScopeLock ScopedLock(&CriticalSection);
+		check(!bHasCancelledRef);
+		bHasCancelledRef = true;
 		check(LiveRequests.Contains(Req));
 		int32 FirstBlock = Offset / FileEntry.CompressionBlockSize;
 		int32 LastBlock = (Offset + BytesToRead - 1) / FileEntry.CompressionBlockSize;
@@ -3706,8 +3707,7 @@ public:
 
 void FPakProcessedReadRequest::CancelRawRequests()
 {
-	bHasCancelled = true;
-	Owner->HandleCanceledRequest(MyCanceledBlocks, this, Offset, BytesToRead);
+	Owner->HandleCanceledRequest(MyCanceledBlocks, this, Offset, BytesToRead, bHasCancelled);
 }
 
 void FPakProcessedReadRequest::GatherResults()
@@ -3724,7 +3724,7 @@ void FPakProcessedReadRequest::GatherResults()
 
 void FPakProcessedReadRequest::DoneWithRawRequests()
 {
-	Owner->RemoveRequest(this, Offset, BytesToRead);
+	Owner->RemoveRequest(this, Offset, BytesToRead, bHasCancelled);
 }
 
 bool FPakProcessedReadRequest::CheckCompletion(const FPakEntry& FileEntry, int32 BlockIndex, TArray<FCachedAsyncBlock*>& Blocks)
