Improvements to pak index memory optimizations
 - Fix entry shrinking not working on pak files that contain files that are encrypted AND compressed
 - Optimize any pak files that are added after encryption keys are registered
 - Improvements to logging / diagnostics
 - Use a unique starting seed (hash of filename) when doing filename unloading. Reduces cost expected collisions (localization paks / encrypted paks with duplicated assets)
 - In non-shipping builds, decode shrunken entries after the shrinking to check they match the original
 - Fix leaked FPakFile object when it was immediately determined to be invalid

diff --git a/Engine/Source/Runtime/PakFile/Private/IPlatformFilePak.cpp b/Engine/Source/Runtime/PakFile/Private/IPlatformFilePak.cpp
--- a/Engine/Source/Runtime/PakFile/Private/IPlatformFilePak.cpp
+++ b/Engine/Source/Runtime/PakFile/Private/IPlatformFilePak.cpp
@@ -4256,6 +4256,8 @@ FPakFile::FPakFile(const TCHAR* Filename, bool bIsSigned)
 	, bIsValid(false)
 	, bFilenamesRemoved(false)
 	, ChunkID(ParseChunkIDFromFilename(Filename))
+	, bAttemptedPakEntryShrink(false)
+	, bAttemptedPakFilenameUnload(false)
  	, MappedFileHandle(nullptr)
 {
 	FArchive* Reader = GetSharedReader(NULL);
@@ -4281,6 +4283,8 @@ FPakFile::FPakFile(IPlatformFile* LowerLevel, const TCHAR* Filename, bool bIsSig
 	, bIsValid(false)
 	, bFilenamesRemoved(false)
 	, ChunkID(ParseChunkIDFromFilename(Filename))
+	, bAttemptedPakEntryShrink(false)
+	, bAttemptedPakFilenameUnload(false)
 	, MappedFileHandle(nullptr)
 {
 	FArchive* Reader = GetSharedReader(LowerLevel);
@@ -4676,18 +4680,23 @@ static inline int32 CDECL CompareFMiniFileEntry(const void* Left, const void* Ri
 bool FPakFile::UnloadPakEntryFilenames(TMap<uint64, FPakEntry>& CrossPakCollisionChecker, TArray<FString>* DirectoryRootsToKeep)
 {
 	// If the process has already been done, get out of here.
-	if (bFilenamesRemoved)
+	if (bAttemptedPakFilenameUnload || bFilenamesRemoved)
 	{
-		return false;
+		return true;
 	}
 
+	UE_LOG(LogPakFile, Log, TEXT("Unloading filenames for pak '%s'"), *PakFilename);
+
 	LLM_SCOPE(ELLMTag::FileSystem);
 
+	// Set this flag so if unloading fails, we don't try again
+	bAttemptedPakFilenameUnload = true;
+
 	// Variables for the filename hashing and collision detection.
 	int NumRetries = 0;
 	const int MAX_RETRIES = 10;
 	bool bHasCollision;
-	FilenameStartHash = 0;
+	FilenameStartHash = FCrc::StrCrc32(*GetFilename());
     
 	// Allocate the temporary array for hashing filenames. The Memset is to hopefully
 	// silence the Visual Studio static analyzer.
@@ -4700,10 +4709,11 @@ bool FPakFile::UnloadPakEntryFilenames(TMap<uint64, FPakEntry>& CrossPakCollisio
 		bHasCollision = false;
 
 		TMap<uint64, FPakEntry> NewCollisionCheckEntries;
+		NewCollisionCheckEntries.Reserve(NumEntries);
 
 		// Build the list of hashes from the Index based on the starting hash.
 		int32 EntryIndex = 0;
-		for (TMap<FString, FPakDirectory>::TConstIterator It(Index); It; ++It)
+		for (TMap<FString, FPakDirectory>::TConstIterator It(Index); (It && !bHasCollision); ++It)
 		{
 			for (FPakDirectory::TConstIterator DirectoryIt(It.Value()); DirectoryIt; ++DirectoryIt)
 			{
@@ -4716,7 +4726,7 @@ bool FPakFile::UnloadPakEntryFilenames(TMap<uint64, FPakEntry>& CrossPakCollisio
 				const FPakEntry* EntryFromPreviousPaks = CrossPakCollisionChecker.Find(FilenameHash);
 				const FPakEntry* EntryFromCurrentPak = NewCollisionCheckEntries.Find(FilenameHash);
 				const FPakEntry& CurrentEntry = Files[DirectoryIt.Value()];
-							   
+
 				if (EntryFromPreviousPaks && (FMemory::Memcmp(EntryFromPreviousPaks->Hash, CurrentEntry.Hash, sizeof(CurrentEntry.Hash))) != 0)
 				{
 					UE_LOG(LogPakFile, Verbose, TEXT("Detected collision with previous pak while hashing %s"), *FinalFilename);
@@ -4739,6 +4749,7 @@ bool FPakFile::UnloadPakEntryFilenames(TMap<uint64, FPakEntry>& CrossPakCollisio
 		{
 			++NumRetries;
 			++FilenameStartHash;
+			UE_LOG(LogPakFile, Verbose, TEXT("Collisions detected. Retrying with new seed..."));
 		}
 		else
 		{
@@ -4755,7 +4766,8 @@ bool FPakFile::UnloadPakEntryFilenames(TMap<uint64, FPakEntry>& CrossPakCollisio
 	if (NumRetries >= MAX_RETRIES)
 	{
 		//		FPlatformMisc::LowLevelOutputDebugString(TEXT("Can't unload pak filenames due to hash collision..."));
-		return true;
+		UE_LOG(LogPakFile, Warning, TEXT("FAILED unloading filenames for pak '%s'"), *PakFilename);
+		return false;
 	}
 
 	// Allocate the storage space.
@@ -4877,19 +4889,24 @@ bool FPakFile::UnloadPakEntryFilenames(TMap<uint64, FPakEntry>& CrossPakCollisio
 		Index.Empty(0);
 	}
 
-	return false;
+	return true;
 }
 
-void FPakFile::ShrinkPakEntriesMemoryUsage()
+bool FPakFile::ShrinkPakEntriesMemoryUsage()
 {
 	// If the process has already been done, get out of here.
-	if (MiniPakEntries != NULL)
+	if (bAttemptedPakEntryShrink || MiniPakEntries != nullptr)
 	{
-		return;
+		return true;
 	}
 
 	LLM_SCOPE(ELLMTag::FileSystem);
 
+	UE_LOG(LogPakFile, Log, TEXT("Shrinking entries for pak '%s'"), *PakFilename);
+
+	// Set this flag so if shrinking fails, we don't try again
+	bAttemptedPakEntryShrink = true;
+
 	// Wander every file entry.
 	int TotalSizeOfCompressedEntries = 0;
 	bool bIsPossibleToShrink = true;
@@ -4901,6 +4918,8 @@ void FPakFile::ShrinkPakEntriesMemoryUsage()
 		bool bIsOffset32BitSafe = Entry.Offset <= MAX_uint32;
 		bool bIsSize32BitSafe = Entry.Size <= MAX_uint32;
 		bool bIsUncompressedSize32BitSafe = Entry.UncompressedSize <= MAX_uint32;
+		uint32 CompressedBlockAlignment = Entry.IsEncrypted() ? FAES::AESBlockSize : 1;
+		int64 HeaderSize = Entry.GetSerializedSize(Info.Version);
 
 		// This data fits into a bitfield (described below), and the data has
 		// to fit within a certain range of bits.
@@ -4921,21 +4940,28 @@ void FPakFile::ShrinkPakEntriesMemoryUsage()
 				bIsPossibleToShrink = false;
 				break;
 			}
-			if (Entry.CompressionBlocks.Num() > 0 && ((Info.HasRelativeCompressedChunkOffsets() ? 0 : Entry.Offset) + Entry.GetSerializedSize(Info.Version) != Entry.CompressionBlocks[0].CompressedStart))
+			if (Entry.CompressionBlocks.Num() > 0 && ((Info.HasRelativeCompressedChunkOffsets() ? 0 : Entry.Offset) + HeaderSize != Entry.CompressionBlocks[0].CompressedStart))
 			{
 				bIsPossibleToShrink = false;
 				break;
 			}
-			if (Entry.CompressionBlocks.Num() == 1 && ((Info.HasRelativeCompressedChunkOffsets() ? 0 : Entry.Offset) + Entry.GetSerializedSize(Info.Version) + Entry.Size != Entry.CompressionBlocks[0].CompressedEnd))
+			if (Entry.CompressionBlocks.Num() == 1)
 			{
-				bIsPossibleToShrink = false;
-				break;
+				uint64 Base = Info.HasRelativeCompressedChunkOffsets() ? 0 : Entry.Offset;
+				uint64 AlignedBlockSize = Align(Entry.CompressionBlocks[0].CompressedEnd - Entry.CompressionBlocks[0].CompressedStart, CompressedBlockAlignment);
+				if ((Base + HeaderSize + Entry.Size) != (Entry.CompressionBlocks[0].CompressedStart + AlignedBlockSize))
+				{
+					bIsPossibleToShrink = false;
+					break;
+				}
 			}
 			if (Entry.CompressionBlocks.Num() > 1)
 			{
 				for (int i = 1; i < Entry.CompressionBlocks.Num(); ++i)
 				{
-					if (Entry.CompressionBlocks[i].CompressedStart != Entry.CompressionBlocks[i - 1].CompressedEnd)
+					uint64 PrevBlockSize = Entry.CompressionBlocks[i - 1].CompressedEnd - Entry.CompressionBlocks[i - 1].CompressedStart;
+					PrevBlockSize = Align(PrevBlockSize, CompressedBlockAlignment);
+					if (Entry.CompressionBlocks[i].CompressedStart != (Entry.CompressionBlocks[i - 1].CompressedStart + PrevBlockSize))
 					{
 						bIsPossibleToShrink = false;
 						break;
@@ -4956,7 +4982,7 @@ void FPakFile::ShrinkPakEntriesMemoryUsage()
 		{
 			TotalSizeOfCompressedEntries +=
 				(bIsSize32BitSafe ? sizeof(uint32) : sizeof(uint64));
-			if (Entry.CompressionBlocks.Num() > 1)
+			if (Entry.CompressionBlocks.Num() > 1 || (Entry.CompressionBlocks.Num() == 1 && Entry.IsEncrypted()))
 			{
 				TotalSizeOfCompressedEntries += Entry.CompressionBlocks.Num() * sizeof(uint32);
 			}
@@ -4965,7 +4991,8 @@ void FPakFile::ShrinkPakEntriesMemoryUsage()
 
 	if (!bIsPossibleToShrink)
 	{
-		return;
+		UE_LOG(LogPakFile, Warning, TEXT("FAILED shrinking entries for pak file '%s'"), *PakFilename);
+		return false;
 	}
 
 	// Allocate the buffer to hold onto all of the bit-encoded compressed FPakEntry structures.
@@ -5063,7 +5090,7 @@ void FPakFile::ShrinkPakEntriesMemoryUsage()
 			}
 
 			// Build the Compression Blocks array.
-			if (FullEntry->CompressionBlocks.Num() > 1)
+			if (FullEntry->CompressionBlocks.Num() > 1 || (FullEntry->CompressionBlocks.Num() == 1 && FullEntry->IsEncrypted()))
 			{
 				for (int CompressionBlockIndex = 0; CompressionBlockIndex < FullEntry->CompressionBlocks.Num(); ++CompressionBlockIndex)
 				{
@@ -5072,6 +5099,13 @@ void FPakFile::ShrinkPakEntriesMemoryUsage()
 				}
 			}
 		}
+
+#if !UE_BUILD_SHIPPING
+		FPakEntry Test;
+		DecodePakEntry(MiniPakEntries + MiniPakEntriesOffsets[EntryIndex], &Test);
+		FMemory::Memcpy(Test.Hash, FullEntry->Hash, 20);
+		check(Test == *FullEntry);
+#endif
 	}
 
 	check(CurrentEntryPtr == MiniPakEntries + TotalSizeOfCompressedEntries);
@@ -5080,7 +5114,7 @@ void FPakFile::ShrinkPakEntriesMemoryUsage()
 	// space of the original anymore.
 	Files.Empty(0);
 
-	return;
+	return true;
 }
 
 #if DO_CHECK
@@ -5657,34 +5691,33 @@ void FPakPlatformFile::InitializeNewAsyncIO()
 void FPakPlatformFile::OptimizeMemoryUsageForMountedPaks()
 {
 #if !(IS_PROGRAM || WITH_EDITOR)
-			FPlatformMisc::LowLevelOutputDebugStringf(TEXT("Checking Pak Config\n"));
-			bool bUnloadPakEntryFilenamesIfPossible = FParse::Param(FCommandLine::Get(), TEXT("unloadpakentryfilenames"));
-			GConfig->GetBool(TEXT("Pak"), TEXT("UnloadPakEntryFilenamesIfPossible"), bUnloadPakEntryFilenamesIfPossible, GEngineIni);
+	bool bUnloadPakEntryFilenamesIfPossible = FParse::Param(FCommandLine::Get(), TEXT("unloadpakentryfilenames"));
+	GConfig->GetBool(TEXT("Pak"), TEXT("UnloadPakEntryFilenamesIfPossible"), bUnloadPakEntryFilenamesIfPossible, GEngineIni);
 
 	if ((bUnloadPakEntryFilenamesIfPossible && !FParse::Param(FCommandLine::Get(), TEXT("nounloadpakentries"))) || FParse::Param(FCommandLine::Get(), TEXT("unloadpakentries")))
-			{
-				// With [Pak] UnloadPakEntryFilenamesIfPossible enabled, [Pak] DirectoryRootsToKeepInMemoryWhenUnloadingPakEntryFilenames
-				// can contain pak entry directory wildcards of which the entire recursive directory structure of filenames underneath a
-				// matching wildcard will be kept.
-				//
-				// Example:
-				//   [Pak]
-				//   DirectoryRootsToKeepInMemoryWhenUnloadingPakEntryFilenames="*/Config/Tags/"
-				//   +DirectoryRootsToKeepInMemoryWhenUnloadingPakEntryFilenames="*/Content/Localization/*"
-				TArray<FString> DirectoryRootsToKeep;
-				GConfig->GetArray(TEXT("Pak"), TEXT("DirectoryRootsToKeepInMemoryWhenUnloadingPakEntryFilenames"), DirectoryRootsToKeep, GEngineIni);
+	{
+		// With [Pak] UnloadPakEntryFilenamesIfPossible enabled, [Pak] DirectoryRootsToKeepInMemoryWhenUnloadingPakEntryFilenames
+		// can contain pak entry directory wildcards of which the entire recursive directory structure of filenames underneath a
+		// matching wildcard will be kept.
+		//
+		// Example:
+		//   [Pak]
+		//   DirectoryRootsToKeepInMemoryWhenUnloadingPakEntryFilenames="*/Config/Tags/"
+		//   +DirectoryRootsToKeepInMemoryWhenUnloadingPakEntryFilenames="*/Content/Localization/*"
+		TArray<FString> DirectoryRootsToKeep;
+		GConfig->GetArray(TEXT("Pak"), TEXT("DirectoryRootsToKeepInMemoryWhenUnloadingPakEntryFilenames"), DirectoryRootsToKeep, GEngineIni);
 
-				FPakPlatformFile* PakPlatformFile = (FPakPlatformFile*)(FPlatformFileManager::Get().FindPlatformFile(FPakPlatformFile::GetTypeName()));
-				PakPlatformFile->UnloadPakEntryFilenames(&DirectoryRootsToKeep);
-			}
+		FPakPlatformFile* PakPlatformFile = (FPakPlatformFile*)(FPlatformFileManager::Get().FindPlatformFile(FPakPlatformFile::GetTypeName()));
+		PakPlatformFile->UnloadPakEntryFilenames(&DirectoryRootsToKeep);
+	}
 
-			bool bShrinkPakEntriesMemoryUsage = FParse::Param(FCommandLine::Get(), TEXT("shrinkpakentries"));
-			GConfig->GetBool(TEXT("Pak"), TEXT("ShrinkPakEntriesMemoryUsage"), bShrinkPakEntriesMemoryUsage, GEngineIni);
-			if (bShrinkPakEntriesMemoryUsage)
-			{
-				FPakPlatformFile* PakPlatformFile = (FPakPlatformFile*)(FPlatformFileManager::Get().FindPlatformFile(FPakPlatformFile::GetTypeName()));
-				PakPlatformFile->ShrinkPakEntriesMemoryUsage();
-			}
+	bool bShrinkPakEntriesMemoryUsage = FParse::Param(FCommandLine::Get(), TEXT("shrinkpakentries"));
+	GConfig->GetBool(TEXT("Pak"), TEXT("ShrinkPakEntriesMemoryUsage"), bShrinkPakEntriesMemoryUsage, GEngineIni);
+	if (bShrinkPakEntriesMemoryUsage)
+	{
+		FPakPlatformFile* PakPlatformFile = (FPakPlatformFile*)(FPlatformFileManager::Get().FindPlatformFile(FPakPlatformFile::GetTypeName()));
+		PakPlatformFile->ShrinkPakEntriesMemoryUsage();
+	}
 #endif
 }
 
@@ -5781,6 +5814,10 @@ bool FPakPlatformFile::Mount(const TCHAR* InPakFilename, uint32 PakOrder, const
 				}
 			}
 		}
+		else
+		{
+			delete Pak;
+		}
 	}
 	else
 	{
@@ -6026,6 +6063,11 @@ void FPakPlatformFile::RegisterEncryptionKey(const FGuid& InGuid, const FAES::FA
 
 		PendingEncryptedPakFiles.RemoveAll([InGuid](const FPakListDeferredEntry& Entry) { return Entry.EncryptionKeyGuid == InGuid; });
 
+		{
+			LLM_SCOPE(ELLMTag::FileSystem);
+			OptimizeMemoryUsageForMountedPaks();
+		}
+
 		UE_CLOG(InGuid.IsValid(), LogPakFile, Log, TEXT("Registered encryption key '%s': %d pak files mounted, %d remain pending"), *InGuid.ToString(), NumMounted, PendingEncryptedPakFiles.Num());
 	}
 }
@@ -6159,43 +6201,70 @@ bool FPakPlatformFile::CopyFile(const TCHAR* To, const TCHAR* From, EPlatformFil
 
 void FPakPlatformFile::UnloadPakEntryFilenames(TArray<FString>* DirectoryRootsToKeep)
 {
-	int32 NumFiles = 0;
+	int32 TotalNumFilenames = 0;
+	int32 NumFilenamesUnloaded = 0;
+	int32 NumPaks = 0;
 	double Timer = 0.0;
 	{
 		SCOPE_SECONDS_COUNTER(Timer);
 
 		TArray<FPakListEntry> Paks;
 		GetMountedPaks(Paks);
-		FPlatformMisc::LowLevelOutputDebugStringf(TEXT("Unloading Pak Entry Filenames\n"));
 		TMap<uint64, FPakEntry> CrossPakCollisionDetector;
-
-		for (const FPakListEntry& Pak : Paks)
+		
+		for (auto Pak : Paks)
 		{
-			NumFiles += Pak.PakFile->GetNumFiles();
+			if (Pak.PakFile->HasFilenames())
+			{
+				TotalNumFilenames += Pak.PakFile->GetNumFiles();
+			}
 		}
 
-		CrossPakCollisionDetector.Reserve(NumFiles);
+		CrossPakCollisionDetector.Reserve(TotalNumFilenames);
 
 		// Sort the pak list by number of entries so that we allow the larger ones a better chance of not encountering a collision
 		Algo::Sort(Paks, [](const FPakListEntry& A, const FPakListEntry& B) { return A.PakFile->GetNumFiles() > B.PakFile->GetNumFiles(); });
 
 		for (auto Pak : Paks)
 		{
-			bool bEncounteredCollisions = Pak.PakFile->UnloadPakEntryFilenames(CrossPakCollisionDetector, DirectoryRootsToKeep);
-			UE_CLOG(bEncounteredCollisions, LogPakFile, Warning, TEXT("Encountered name collisions while unloading pak index strings for %s"), *Pak.PakFile->GetFilename());
+			if (Pak.PakFile->HasFilenames())
+			{
+				NumPaks++;
+				int32 NumFilesInThisPak = Pak.PakFile->GetNumFiles();
+
+				if (Pak.PakFile->UnloadPakEntryFilenames(CrossPakCollisionDetector, DirectoryRootsToKeep))
+				{
+					NumFilenamesUnloaded += NumFilesInThisPak;
+				}
+			}
 		}
 	}
-	UE_LOG(LogPakFile, Log, TEXT("PakEntry filenames unloaded %d filenames in %.4fs"), NumFiles, Timer);
+	UE_LOG(LogPakFile, Log, TEXT("Unloaded %d/%d filenames from %d pak files in %.4fs"), NumFilenamesUnloaded, TotalNumFilenames, NumPaks, Timer);
 }
 
 void FPakPlatformFile::ShrinkPakEntriesMemoryUsage()
 {
-	TArray<FPakListEntry> Paks;
-	GetMountedPaks(Paks);
-	for (auto Pak : Paks)
+	double Timer = 0.0;
+	int32 NumPakFiles = 0;
+	int32 NumEntries = 0;
 	{
-		Pak.PakFile->ShrinkPakEntriesMemoryUsage();
+		SCOPE_SECONDS_COUNTER(Timer);
+
+		TArray<FPakListEntry> Paks;
+		GetMountedPaks(Paks);
+		for (auto Pak : Paks)
+		{
+			if (!Pak.PakFile->HasShrunkPakEntries())
+			{
+				NumPakFiles++;
+				if (Pak.PakFile->ShrinkPakEntriesMemoryUsage())
+				{
+					NumEntries += Pak.PakFile->GetNumFiles();
+				}
+			}
+		}
 	}
+	UE_LOG(LogPakFile, Log, TEXT("Shrunk %d entries from %d pak files in %.4fs"), NumEntries, NumPakFiles, Timer);
 }
 
 /**
diff --git a/Engine/Source/Runtime/PakFile/Public/IPlatformFilePak.h b/Engine/Source/Runtime/PakFile/Public/IPlatformFilePak.h
--- a/Engine/Source/Runtime/PakFile/Public/IPlatformFilePak.h
+++ b/Engine/Source/Runtime/PakFile/Public/IPlatformFilePak.h
@@ -549,6 +549,10 @@ class PAKFILE_API FPakFile : FNoncopyable
 	bool bFilenamesRemoved;
 	/** ID for the chunk this pakfile is part of. INDEX_NONE if this isn't a pak chunk (derived from filename) */
 	int32 ChunkID;
+	/** Flag to say we tried shrinking pak entries already */
+	bool bAttemptedPakEntryShrink;
+	/** Flag to say we tried unloading pak index filenames already */
+	bool bAttemptedPakFilenameUnload;
 
 	class IMappedFileHandle* MappedFileHandle;
 	FCriticalSection MappedFileHandleCriticalSection;
@@ -656,7 +660,7 @@ public:
 	 */
 	int32 GetNumFiles() const
 	{
-		return Files.Num();
+		return NumEntries;
 	}
 
 	void GetFilenames(TArray<FString>& OutFileList) const;
@@ -949,7 +953,7 @@ public:
 	/**
 	 * Saves memory by hashing the filenames, if possible. After this process,
 	 * wildcard scanning of pak entries can no longer be performed. Returns TRUE
-	 * if there were any collisions within this pak or with any of the previous pak results supplied in CrossPakCollisionChecker
+	 * if the process successfully unloaded filenames from this pak
 	 *
 	 * @param CrossPakCollisionChecker A map of hash->fileentry records encountered during filename unloading on other pak files. Used to detect collisions with entries in other pak files.
 	 * @param DirectoryRootsToKeep An array of strings in wildcard format that specify whole directory structures of filenames to keep in memory for directory iteration to work.
@@ -959,7 +963,15 @@ public:
 	/**
 	 * Lower memory usage by bit-encoding the pak file entry information.
 	 */
-	void ShrinkPakEntriesMemoryUsage();
+	bool ShrinkPakEntriesMemoryUsage();
+
+	/**
+	 * Returns whether the pak files list has been shrunk or not
+	 */
+	bool HasShrunkPakEntries() const
+	{
+		return bAttemptedPakEntryShrink;
+	}
 
 private:
 
@@ -1052,15 +1064,19 @@ private:
 		// Filter the encrypted flag.
 		OutEntry->SetEncrypted((Value & (1 << 22)) != 0);
 
-		// Filter the compression block size or use the UncompressedSize if less that 64k.
-		OutEntry->CompressionBlockSize = OutEntry->UncompressedSize < 65536 ? (uint32)OutEntry->UncompressedSize : ((Value & 0x3f) << 11);
-
 		// This should clear out any excess CompressionBlocks that may be valid in the user's
 		// passed in entry.
 		uint32 CompressionBlocksCount = (Value >> 6) & 0xffff;
 		OutEntry->CompressionBlocks.Empty(CompressionBlocksCount);
 		OutEntry->CompressionBlocks.SetNum(CompressionBlocksCount);
 
+		// Filter the compression block size or use the UncompressedSize if less that 64k.
+		OutEntry->CompressionBlockSize = 0;
+		if (CompressionBlocksCount > 0)
+		{
+			OutEntry->CompressionBlockSize = OutEntry->UncompressedSize < 65536 ? (uint32)OutEntry->UncompressedSize : ((Value & 0x3f) << 11);
+		}
+
 		// Set Verified to true to avoid have a synchronous open fail comparing FPakEntry structures.
 		OutEntry->Verified = true;
 
@@ -1071,7 +1087,7 @@ private:
 		int64 BaseOffset = Info.HasRelativeCompressedChunkOffsets() ? 0 : OutEntry->Offset;
 
 		// Handle building of the CompressionBlocks array.
-		if (OutEntry->CompressionBlocks.Num() == 1)
+		if (OutEntry->CompressionBlocks.Num() == 1 && !OutEntry->IsEncrypted())
 		{
 			// If the number of CompressionBlocks is 1, we didn't store any extra information.
 			// Derive what we can from the entry's file offset and size.
@@ -1084,6 +1100,9 @@ private:
 			// Get the right pointer to start copying the CompressionBlocks information from.
 			uint32* CompressionBlockSizePtr = (uint32*)SourcePtr;
 
+			// Alignment of the compressed blocks
+			uint64 CompressedBlockAlignment = OutEntry->IsEncrypted() ? FAES::AESBlockSize : 1;
+
 			// CompressedBlockOffset is the starting offset. Everything else can be derived from there.
 			int64 CompressedBlockOffset = BaseOffset + OutEntry->GetSerializedSize(Info.Version);
 			for (int CompressionBlockIndex = 0; CompressionBlockIndex < OutEntry->CompressionBlocks.Num(); ++CompressionBlockIndex)
@@ -1091,7 +1110,7 @@ private:
 				FPakCompressedBlock& CompressedBlock = OutEntry->CompressionBlocks[CompressionBlockIndex];
 				CompressedBlock.CompressedStart = CompressedBlockOffset;
 				CompressedBlock.CompressedEnd = CompressedBlockOffset + *CompressionBlockSizePtr++;
-				CompressedBlockOffset = CompressedBlock.CompressedEnd;
+				CompressedBlockOffset += Align(CompressedBlock.CompressedEnd - CompressedBlock.CompressedStart, CompressedBlockAlignment);
 			}
 		}
 
