Fix crash while shrinking pak entries for delete records

diff --git a/Engine/Source/Runtime/PakFile/Private/IPlatformFilePak.cpp b/Engine/Source/Runtime/PakFile/Private/IPlatformFilePak.cpp
--- a/Engine/Source/Runtime/PakFile/Private/IPlatformFilePak.cpp
+++ b/Engine/Source/Runtime/PakFile/Private/IPlatformFilePak.cpp
@@ -5101,10 +5101,13 @@ bool FPakFile::ShrinkPakEntriesMemoryUsage()
 		}
 
 #if !UE_BUILD_SHIPPING
-		FPakEntry Test;
-		DecodePakEntry(MiniPakEntries + MiniPakEntriesOffsets[EntryIndex], &Test);
-		FMemory::Memcpy(Test.Hash, FullEntry->Hash, 20);
-		check(Test == *FullEntry);
+		if (!FullEntry->IsDeleteRecord())
+		{
+			FPakEntry Test;
+			DecodePakEntry(MiniPakEntries + MiniPakEntriesOffsets[EntryIndex], &Test);
+			FMemory::Memcpy(Test.Hash, FullEntry->Hash, 20);
+			check(Test == *FullEntry);
+		}
 #endif
 	}
 
