UE4 - Added a panic recovery mode to failures to decompress data. This is a workaround for a mysterious failure here with no known repro.

diff --git a/Engine/Source/Runtime/PakFile/Private/IPlatformFilePak.cpp b/Engine/Source/Runtime/PakFile/Private/IPlatformFilePak.cpp
--- a/Engine/Source/Runtime/PakFile/Private/IPlatformFilePak.cpp
+++ b/Engine/Source/Runtime/PakFile/Private/IPlatformFilePak.cpp
@@ -324,6 +324,15 @@ void DecryptData(uint8* InData, uint32 InDataSize, FGuid InEncryptionKeyGuid)
 	}
 }
 
+#if !UE_BUILD_SHIPPING
+static int32 GPakCache_ForceDecompressionFails = 0;
+static FAutoConsoleVariableRef CVar_ForceDecompressionFails(
+	TEXT("ForceDecompressionFails"),
+	GPakCache_ForceDecompressionFails,
+	TEXT("If > 0, then force decompression failures to test the panic sync read fallback.")
+);
+#endif
+
 #if USE_PAK_PRECACHE
 #include "Async/TaskGraphInterfaces.h"
 #define PAK_CACHE_GRANULARITY (64*1024)
@@ -2720,6 +2729,7 @@ protected:
 	int64 BytesToRead;
 	FEvent* WaitEvent;
 	FCachedAsyncBlock* BlockPtr;
+	FName PanicPakFile;
 	EAsyncIOPriorityAndFlags PriorityAndFlags;
 	bool bRequestOutstanding;
 	bool bNeedsRemoval;
@@ -2732,6 +2742,7 @@ public:
 		, BytesToRead(InBytesToRead)
 		, WaitEvent(nullptr)
 		, BlockPtr(InBlockPtr)
+		, PanicPakFile(InPakFile)
 		, PriorityAndFlags(InPriorityAndFlags)
 		, bRequestOutstanding(true)
 		, bNeedsRemoval(true)
@@ -2851,6 +2862,22 @@ public:
 			SetAllComplete();
 		}
 	}
+
+	void PanicSyncRead(uint8* Buffer)
+	{
+		IFileHandle* Handle = IPlatformFile::GetPlatformPhysical().OpenRead(*PanicPakFile.ToString());
+		UE_CLOG(!Handle, LogPakFile, Fatal, TEXT("PanicSyncRead failed to open pak file %s"), *PanicPakFile.ToString());
+		if (!Handle->Seek(Offset))
+		{
+			UE_LOG(LogPakFile, Fatal, TEXT("PanicSyncRead failed to seek pak file %s   %d bytes at %lld "), *PanicPakFile.ToString(), BytesToRead, Offset);
+		}
+
+		if (!Handle->Read(Buffer, BytesToRead))
+		{
+			UE_LOG(LogPakFile, Fatal, TEXT("PanicSyncRead failed to read pak file %s   %d bytes at %lld "), *PanicPakFile.ToString(), BytesToRead, Offset);
+		}	
+		delete Handle;
+	}
 };
 
 class FPakEncryptedReadRequest : public FPakReadRequestBase
@@ -3477,6 +3504,22 @@ public:
 		{
 			check(Block.Raw && Block.RawSize && !Block.Processed);
 
+#if !UE_BUILD_SHIPPING
+			bool bCorrupted = false;
+			if (GPakCache_ForceDecompressionFails && FMath::FRand() < 0.001f)
+			{
+				int32 CorruptOffset = FMath::Clamp(int32(FMath::FRandRange(0, Block.RawSize - 1)), 0, Block.RawSize - 1);
+				uint8 CorruptValue = uint8(FMath::Clamp(int32(FMath::FRandRange(0, 255)), 0, 255));
+				if (Block.Raw[CorruptOffset] != CorruptValue)
+				{
+					UE_LOG(LogPakFile, Error, TEXT("Forcing corruption of decompression source data (predecryption) to verify panic read recovery.  Offset = %d, Value = 0x%x"), CorruptOffset, int32(CorruptValue));
+					Block.Raw[CorruptOffset] = CorruptValue;
+					bCorrupted = true;
+				}
+			}
+#endif
+
+
 			if (FileEntry.IsEncrypted())
 			{
 				INC_DWORD_STAT(STAT_PakCache_CompressedDecrypts);
@@ -3496,10 +3539,63 @@ public:
 				check(Block.DecompressionRawSize == Block.RawSize);
 			}
 
-			if( !FCompression::UncompressMemory(CompressionMethod, Output, Block.ProcessedSize, Block.Raw, Block.DecompressionRawSize) )
+			bool bFailed = !FCompression::UncompressMemory(CompressionMethod, Output, Block.ProcessedSize, Block.Raw, Block.DecompressionRawSize);
+
+#if !UE_BUILD_SHIPPING
+			if (bCorrupted && !bFailed)
+			{
+				UE_LOG(LogPakFile, Error, TEXT("The payload was corrupted, but this did not trigger a decompression failed.....pretending it failed anyway because otherwise it can crash later."));
+				bFailed = true;
+			}
+#endif
+
+			if (bFailed)
 			{
-				const FString HexBytes = BytesToHex(Block.Raw,FMath::Min(Block.DecompressionRawSize,32));
-				UE_LOG( LogPakFile, Fatal, TEXT("Pak Decompression failed. PakFile:%s, EntryOffset:%lld, EntrySize:%lld, Method:%s, ProcessedSize:%d, RawSize:%d, Crc32:%u, BlockIndex:%d, Encrypt:%d, Delete:%d, Output:%p, Raw:%p, Processed:%p, Bytes:[%s...]"), *PakFile.ToString(), FileEntry.Offset, FileEntry.Size, *CompressionMethod.ToString(), Block.ProcessedSize, Block.DecompressionRawSize, FCrc::MemCrc32( Block.Raw, Block.DecompressionRawSize ), Block.BlockIndex, FileEntry.IsEncrypted()?1:0, FileEntry.IsDeleteRecord()?1:0, Output, Block.Raw, Block.Processed, *HexBytes );
+				{
+					const FString HexBytes = BytesToHex(Block.Raw, FMath::Min(Block.DecompressionRawSize, 32));
+					UE_LOG(LogPakFile, Error, TEXT("Pak Decompression failed. PakFile:%s, EntryOffset:%lld, EntrySize:%lld, Method:%s, ProcessedSize:%d, RawSize:%d, Crc32:%u, BlockIndex:%d, Encrypt:%d, Delete:%d, Output:%p, Raw:%p, Processed:%p, Bytes:[%s...]"), *PakFile.ToString(), FileEntry.Offset, FileEntry.Size, *CompressionMethod.ToString(), Block.ProcessedSize, Block.DecompressionRawSize, FCrc::MemCrc32(Block.Raw, Block.DecompressionRawSize), Block.BlockIndex, FileEntry.IsEncrypted() ? 1 : 0, FileEntry.IsDeleteRecord() ? 1 : 0, Output, Block.Raw, Block.Processed, *HexBytes);
+				}
+				uint8* TempBuffer = (uint8*)FMemory::Malloc(Block.RawSize);
+				{
+					FScopeLock ScopedLock(&CriticalSection);
+					UE_CLOG(!Block.RawRequest, LogPakFile, Fatal, TEXT("Cannot retry because Block.RawRequest is null."));
+
+					Block.RawRequest->PanicSyncRead(TempBuffer);
+				}
+
+				if (FileEntry.IsEncrypted())
+				{
+					DecryptData(TempBuffer, Block.RawSize, EncryptionKeyGuid);
+				}
+				if (FMemory::Memcmp(TempBuffer, Block.Raw, Block.DecompressionRawSize) != 0)
+				{
+					UE_LOG(LogPakFile, Warning, TEXT("Panic re-read (and decrypt if applicable) resulted in a different buffer."));
+
+					int32 Offset = 0;
+					for (; Offset < Block.DecompressionRawSize; Offset++)
+					{
+						if (TempBuffer[Offset] != Block.Raw[Offset])
+						{
+							break;
+						}
+					}
+					UE_CLOG(Offset >= Block.DecompressionRawSize, LogPakFile, Fatal, TEXT("Buffers were different yet all bytes were the same????"));
+
+					UE_LOG(LogPakFile, Warning, TEXT("Buffers differ at offset %d."), Offset);
+					const FString HexBytes1 = BytesToHex(Block.Raw + Offset, FMath::Min(Block.DecompressionRawSize - Offset, 64));
+					UE_LOG(LogPakFile, Warning, TEXT("Original read (and decrypt) %s"), *HexBytes1);
+					const FString HexBytes2 = BytesToHex(TempBuffer + Offset, FMath::Min(Block.DecompressionRawSize - Offset, 64));
+					UE_LOG(LogPakFile, Warning, TEXT("Panic reread  (and decrypt) %s"), *HexBytes2);
+				}
+				if (!FCompression::UncompressMemory(CompressionMethod, Output, Block.ProcessedSize, TempBuffer, Block.DecompressionRawSize))
+				{
+					UE_LOG(LogPakFile, Fatal, TEXT("Retry was NOT sucessful."));
+				}
+				else
+				{
+					UE_LOG(LogPakFile, Warning, TEXT("Retry was sucessful."));
+				}
+				FMemory::Free(TempBuffer);
 			}
 			FMemory::Free(Block.Raw);
 			Block.Raw = nullptr;
