Duplicating 4461277
Workaround for UE-58013, though not a full fix.  Requires setting bAllowNonUFSIniWhenCooked in the Target.cs and using a cooked build with pak files.

diff --git a/Engine/Source/Programs/UnrealBuildTool/Configuration/TargetRules.cs b/Engine/Source/Programs/UnrealBuildTool/Configuration/TargetRules.cs
--- a/Engine/Source/Programs/UnrealBuildTool/Configuration/TargetRules.cs
+++ b/Engine/Source/Programs/UnrealBuildTool/Configuration/TargetRules.cs
@@ -935,10 +935,15 @@ namespace UnrealBuildTool
 		[CommandLine("-ToolChain")]
 		public string ToolChainName = null;
 
-        /// <summary>
-        /// Whether to load generated ini files in cooked build
-        /// </summary>
-        public bool bAllowGeneratedIniWhenCooked = true;
+		/// <summary>
+		/// Whether to load generated ini files in cooked build, (GameUserSettings.ini loaded either way)
+		/// </summary>
+		public bool bAllowGeneratedIniWhenCooked = true;
+
+		/// <summary>
+		/// Whether to load non-ufs ini files in cooked build, (GameUserSettings.ini loaded either way)
+		/// </summary>
+		public bool bAllowNonUFSIniWhenCooked = true;
 
 		/// <summary>
 		/// Add all the public folders as include paths for the compile environment.
diff --git a/Engine/Source/Programs/UnrealBuildTool/Configuration/UEBuildTarget.cs b/Engine/Source/Programs/UnrealBuildTool/Configuration/UEBuildTarget.cs
--- a/Engine/Source/Programs/UnrealBuildTool/Configuration/UEBuildTarget.cs
+++ b/Engine/Source/Programs/UnrealBuildTool/Configuration/UEBuildTarget.cs
@@ -457,6 +457,11 @@ namespace UnrealBuildTool
 				RulesObject.bBuildRequiresCookedData = true;
 			}
 
+			if (!RulesObject.bAllowNonUFSIniWhenCooked)
+			{
+				RulesObject.GlobalDefinitions.Add("DISABLE_NONUFS_INI_WHEN_COOKED=1");
+			}
+
 			// Allow the platform to finalize the settings
 			UEBuildPlatform Platform = UEBuildPlatform.GetBuildPlatform(RulesObject.Platform);
 			Platform.ValidateTarget(RulesObject);
diff --git a/Engine/Source/Runtime/PakFile/Private/IPlatformFilePak.cpp b/Engine/Source/Runtime/PakFile/Private/IPlatformFilePak.cpp
--- a/Engine/Source/Runtime/PakFile/Private/IPlatformFilePak.cpp
+++ b/Engine/Source/Runtime/PakFile/Private/IPlatformFilePak.cpp
@@ -30,6 +30,10 @@ DEFINE_LOG_CATEGORY(LogPakFile);
 DEFINE_STAT(STAT_PakFile_Read);
 DEFINE_STAT(STAT_PakFile_NumOpenHandles);
 
+#ifndef DISABLE_NONUFS_INI_WHEN_COOKED
+#define DISABLE_NONUFS_INI_WHEN_COOKED 0
+#endif
+
 int32 ParseChunkIDFromFilename(const FString& InFilename)
 {
 	FString ChunkIdentifier(TEXT("pakchunk"));
@@ -3896,6 +3900,13 @@ bool FPakPlatformFile::IsNonPakFilenameAllowed(const FString& InFilename)
 	}
 #endif
 
+#if DISABLE_NONUFS_INI_WHEN_COOKED
+	if (FPlatformProperties::RequiresCookedData() && InFilename.EndsWith(IniFileExtension) && !InFilename.EndsWith(GameUserSettingsIniFilename))
+	{
+		bAllowed = false;
+	}
+#endif
+
 	FFilenameSecurityDelegate& FilenameSecurityDelegate = GetFilenameSecurityDelegate();
 	if (bAllowed)
 	{
@@ -4939,6 +4950,11 @@ bool FPakPlatformFile::Initialize(IPlatformFile* Inner, const TCHAR* CmdLine)
 	ExcludedNonPakExtensions.Add(TEXT("uexp"));
 #endif
 
+#if DISABLE_NONUFS_INI_WHEN_COOKED
+	IniFileExtension = TEXT(".ini");
+	GameUserSettingsIniFilename = TEXT("GameUserSettings.ini");
+#endif
+
 	FEncryptionKey DecryptionKey;
 	GetPakSigningKeys(DecryptionKey);
 
diff --git a/Engine/Source/Runtime/PakFile/Public/IPlatformFilePak.h b/Engine/Source/Runtime/PakFile/Public/IPlatformFilePak.h
--- a/Engine/Source/Runtime/PakFile/Public/IPlatformFilePak.h
+++ b/Engine/Source/Runtime/PakFile/Public/IPlatformFilePak.h
@@ -1311,6 +1311,10 @@ class PAKFILE_API FPakPlatformFile : public IPlatformFile
 	mutable FCriticalSection PakListCritical;
 	/** Cache of extensions that we automatically reject if not found in pak file */
 	TSet<FName> ExcludedNonPakExtensions;
+	/** The extension used for ini files, used for excluding ini files */
+	FString IniFileExtension;
+	/** The filename for the gameusersettings ini file, used for excluding ini files, but not gameusersettings */
+	FString GameUserSettingsIniFilename;
 
 	/**
 	 * Gets mounted pak files
