Integrated a fix from UDN where the engine would crash when trying to load a very small encrypted file (<16bytes) from a pak file, where the read address wasn't already aligned to the AES block size.

(https://udn.unrealengine.com/questions/431989/crash-while-reading-a-very-small-file-in-encrypted.html)

diff --git a/Engine/Source/Runtime/PakFile/Public/IPlatformFilePak.h b/Engine/Source/Runtime/PakFile/Public/IPlatformFilePak.h
--- a/Engine/Source/Runtime/PakFile/Public/IPlatformFilePak.h
+++ b/Engine/Source/Runtime/PakFile/Public/IPlatformFilePak.h
@@ -961,27 +961,29 @@ public:
 
 	void Serialize(int64 DesiredPosition, void* V, int64 Length)
 	{
-		uint8 TempBuffer[EncryptionPolicy::Alignment];
+		const constexpr int64 Alignment = (int64)EncryptionPolicy::Alignment;
+		const constexpr int64 AlignmentMask = ~(Alignment - 1);
+		uint8 TempBuffer[Alignment];
 		if (EncryptionPolicy::AlignReadRequest(DesiredPosition) != DesiredPosition)
 		{
-			int64 Start = DesiredPosition & ~(EncryptionPolicy::Alignment-1);
+			int64 Start = DesiredPosition & AlignmentMask;
 			int64 Offset = DesiredPosition - Start;
-			int32 CopySize = EncryptionPolicy::Alignment-(DesiredPosition-Start);
+			int64 CopySize = FMath::Min(Alignment - Offset, Length);
 			PakReader->Seek(OffsetToFile + Start);
-			PakReader->Serialize(TempBuffer, EncryptionPolicy::Alignment);
-			EncryptionPolicy::DecryptBlock(TempBuffer, EncryptionPolicy::Alignment);
-			FMemory::Memcpy(V, TempBuffer+Offset, CopySize);
+			PakReader->Serialize(TempBuffer, Alignment);
+			EncryptionPolicy::DecryptBlock(TempBuffer, Alignment);
+			FMemory::Memcpy(V, TempBuffer + Offset, CopySize);
 			V = (void*)((uint8*)V + CopySize);
 			DesiredPosition += CopySize;
 			Length -= CopySize;
-			check(DesiredPosition % EncryptionPolicy::Alignment == 0);
+			check(Length == 0 || DesiredPosition % Alignment == 0);
 		}
 		else
 		{
 			PakReader->Seek(OffsetToFile + DesiredPosition);
 		}
 		
-		int64 CopySize = Length & ~(EncryptionPolicy::Alignment-1);
+		int64 CopySize = Length & AlignmentMask;
 		PakReader->Serialize(V, CopySize);
 		EncryptionPolicy::DecryptBlock(V, CopySize);
 		Length -= CopySize;
@@ -989,8 +991,8 @@ public:
 
 		if (Length > 0)
 		{
-			PakReader->Serialize(TempBuffer, EncryptionPolicy::Alignment);
-			EncryptionPolicy::DecryptBlock(TempBuffer, EncryptionPolicy::Alignment);
+			PakReader->Serialize(TempBuffer, Alignment);
+			EncryptionPolicy::DecryptBlock(TempBuffer, Alignment);
 			FMemory::Memcpy(V, TempBuffer, Length);
 		}
 	}
