Fix for compressed file offset being stored as an int32 when it should have been int64, breaking pak files >2GB. Already fixed in other branches, but copying the fix here

diff --git a/Engine/Source/Runtime/PakFile/Private/IPlatformFilePak.cpp b/Engine/Source/Runtime/PakFile/Private/IPlatformFilePak.cpp
--- a/Engine/Source/Runtime/PakFile/Private/IPlatformFilePak.cpp
+++ b/Engine/Source/Runtime/PakFile/Private/IPlatformFilePak.cpp
@@ -3043,7 +3043,7 @@ class FPakAsyncReadFileHandle final : public IAsyncReadFileHandle
 	FAsyncFileCallBack ReadCallbackFunction;
 	FCriticalSection CriticalSection;
 	int32 NumLiveRawRequests;
-	int32 CompressedChunkOffset;
+	int64 CompressedChunkOffset;
 
 	TMap<FCachedAsyncBlock*, FPakProcessedReadRequest*> OutstandingCancelMapBlock;
 
