Fixed some incorrect documentation for pakfile compressed chunk headers

diff --git a/Engine/Source/Runtime/PakFile/Public/IPlatformFilePak.h b/Engine/Source/Runtime/PakFile/Public/IPlatformFilePak.h
--- a/Engine/Source/Runtime/PakFile/Public/IPlatformFilePak.h
+++ b/Engine/Source/Runtime/PakFile/Public/IPlatformFilePak.h
@@ -137,9 +137,9 @@ struct FPakInfo
  */
 struct FPakCompressedBlock
 {
-	/** Offset of the start of a compression block. Offset is absolute. */
+	/** Offset of the start of a compression block. Offset is relative to the start of the compressed chunk data */
 	int64 CompressedStart;
-	/** Offset of the end of a compression block. This may not align completely with the start of the next block. Offset is absolute. */
+	/** Offset of the end of a compression block. This may not align completely with the start of the next block. Offset is relative to the start of the compressed chunk data. */
 	int64 CompressedEnd;
 
 	bool operator == (const FPakCompressedBlock& B) const
