Fix an error in FPakEntry::GetSerializedSize() which was incorrect for versions before the named compression mode stuff

diff --git a/Engine/Source/Runtime/PakFile/Public/IPlatformFilePak.h b/Engine/Source/Runtime/PakFile/Public/IPlatformFilePak.h
--- a/Engine/Source/Runtime/PakFile/Public/IPlatformFilePak.h
+++ b/Engine/Source/Runtime/PakFile/Public/IPlatformFilePak.h
@@ -310,7 +310,17 @@ struct FPakEntry
 	 */
 	int64 GetSerializedSize(int32 Version) const
 	{
-		int64 SerializedSize = sizeof(Offset) + sizeof(Size) + sizeof(UncompressedSize) + sizeof(CompressionMethodIndex) + sizeof(Hash);
+		int64 SerializedSize = sizeof(Offset) + sizeof(Size) + sizeof(UncompressedSize) + sizeof(Hash);
+
+		if (Version >= FPakInfo::PakFile_Version_FNameBasedCompressionMethod)
+		{
+			SerializedSize += sizeof(CompressionMethodIndex);
+		}
+		else
+		{
+			SerializedSize += sizeof(int32); // Old CompressedMethod var from pre-fname based compression methods
+		}
+
 		if (Version >= FPakInfo::PakFile_Version_CompressionEncryption)
 		{
 			SerializedSize += sizeof(Flags) + sizeof(CompressionBlockSize);
