Removed ensure from signedarchivereader pak signature failure handler to make it consistent with the pakprecacher version. Games can override the signature check failure delegate to choose what to do when a signature fails

diff --git a/Engine/Source/Runtime/PakFile/Private/IPlatformFilePak.cpp b/Engine/Source/Runtime/PakFile/Private/IPlatformFilePak.cpp
--- a/Engine/Source/Runtime/PakFile/Private/IPlatformFilePak.cpp
+++ b/Engine/Source/Runtime/PakFile/Private/IPlatformFilePak.cpp
@@ -3168,7 +3168,6 @@ void FPakPrecacher::DoSignatureCheck(bool bWasCanceled, IAsyncReadRequest* Reque
 				FPakData* PakData = &CachedPakData[PakIndex];
 
 				UE_LOG(LogPakFile, Warning, TEXT("Pak chunk signing mismatch on chunk [%i/%i]! Expected 0x%8X, Received 0x%8X"), SignatureIndex, PakData->ChunkHashes.Num(), PakData->OriginalSignatureFileHash, ThisHash);
-				UE_LOG(LogPakFile, Warning, TEXT("Pak file has been corrupted or tampered with!"));
 
 				// Check the signatures are still as we expected them
 				TPakChunkHash CurrentSignatureHash = ComputePakChunkHash(&PakData->ChunkHashes[0], PakData->ChunkHashes.Num() * sizeof(TPakChunkHash));
diff --git a/Engine/Source/Runtime/PakFile/Private/SignedArchiveReader.cpp b/Engine/Source/Runtime/PakFile/Private/SignedArchiveReader.cpp
--- a/Engine/Source/Runtime/PakFile/Private/SignedArchiveReader.cpp
+++ b/Engine/Source/Runtime/PakFile/Private/SignedArchiveReader.cpp
@@ -281,15 +281,12 @@ bool FChunkCacheWorker::CheckSignature(const FChunkRequest& ChunkInfo)
 	if (!bChunkHashesMatch)
 	{
 		UE_LOG(LogPakFile, Warning, TEXT("Pak chunk signing mismatch on chunk [%i/%i]! Expected 0x%8X, Received 0x%8X"), ChunkInfo.Index, ChunkHashes.Num(), ChunkHashes[ChunkInfo.Index], ChunkHash);
-		UE_LOG(LogPakFile, Warning, TEXT("Pak file has been corrupted or tampered with!"));
 
 		TPakChunkHash CurrentSignatureHash = ComputePakChunkHash(&ChunkHashes[0], ChunkHashes.Num() * sizeof(TPakChunkHash));
 		if (OriginalSignatureFileHash != CurrentSignatureHash)
 		{
 			UE_LOG(LogPakFile, Warning, TEXT("Master signature table has changed since initialization!"));
 		}
-			
-		ensure(bChunkHashesMatch);
 
 		const FPakChunkSignatureCheckFailedData Data(Reader->GetArchiveName(), ChunkHashes[ChunkInfo.Index], ChunkHash, ChunkInfo.Index);
 		FPakPlatformFile::GetPakChunkSignatureCheckFailedHandler().Broadcast(Data);
