Pak file handles no longer store a pointer to a shared reader archive. Those readers are thread specific and means handles can't be read from different threads without causing problems. By storing a function that looks up the pak reader dynamically on each file operation, the requests will work properly from any thread.

diff --git a/Engine/Source/Runtime/PakFile/Private/IPlatformFilePak.cpp b/Engine/Source/Runtime/PakFile/Private/IPlatformFilePak.cpp
--- a/Engine/Source/Runtime/PakFile/Private/IPlatformFilePak.cpp
+++ b/Engine/Source/Runtime/PakFile/Private/IPlatformFilePak.cpp
@@ -3828,10 +3828,10 @@ public:
 		}
 	};
 
-	FPakCompressedReaderPolicy(const FPakFile& InPakFile, const FPakEntry& InPakEntry, FArchive* InPakReader)
+	FPakCompressedReaderPolicy(const FPakFile& InPakFile, const FPakEntry& InPakEntry, TAcquirePakReaderFunction& InAcquirePakReader)
 		: PakFile(InPakFile)
 		, PakEntry(InPakEntry)
-		, PakReader(InPakReader)
+		, AcquirePakReader(InAcquirePakReader)
 	{
 	}
 
@@ -3839,8 +3839,8 @@ public:
 	const FPakFile&		PakFile;
 	/** Pak file entry for this file. */
 	FPakEntry			PakEntry;
-	/** Pak file archive to read the data from. */
-	FArchive*			PakReader;
+	/** Function that gives us an FArchive to read from. The result should never be cached, but acquired and used within the function doing the serialization operation */
+	TAcquirePakReaderFunction AcquirePakReader;
 
 	FORCEINLINE int64 FileSize() const
 	{
@@ -3874,6 +3874,8 @@ public:
 		WorkingBuffers[0] = ScratchSpace.ScratchBuffer.Get();
 		WorkingBuffers[1] = ScratchSpace.ScratchBuffer.Get() + WorkingBufferRequiredSize;
 
+		FArchive* PakReader = AcquirePakReader();
+
 		while (Length > 0)
 		{
 			const FPakCompressedBlock& Block = PakEntry.CompressionBlocks[CompressionBlockIndex];
@@ -4831,28 +4833,27 @@ FArchive* FPakFile::GetSharedReader(IPlatformFile* LowerLevel)
 		{
 			PakReader = ExistingReader->Get();
 		}
-	}
-	if (!PakReader)
-	{
-		// Create a new FArchive reader and pass it to the new handle.
-		if (LowerLevel != NULL)
+
+		if (!PakReader)
 		{
-			IFileHandle* PakHandle = LowerLevel->OpenRead(*GetFilename());
-			if (PakHandle)
+			// Create a new FArchive reader and pass it to the new handle.
+			if (LowerLevel != NULL)
 			{
-				PakReader = CreatePakReader(*PakHandle, *GetFilename());
+				IFileHandle* PakHandle = LowerLevel->OpenRead(*GetFilename());
+				if (PakHandle)
+				{
+					PakReader = CreatePakReader(*PakHandle, *GetFilename());
+				}
 			}
-		}
-		else
-		{
-			PakReader = CreatePakReader(*GetFilename());
-		}
-		if (!PakReader)
-		{
-			UE_LOG(LogPakFile, Fatal, TEXT("Unable to create pak \"%s\" handle"), *GetFilename());
-		}
-		{
-			FScopeLock ScopedLock(&CriticalSection);
+			else
+			{
+				PakReader = CreatePakReader(*GetFilename());
+			}
+			if (!PakReader)
+			{
+				UE_LOG(LogPakFile, Fatal, TEXT("Unable to create pak \"%s\" handle"), *GetFilename());
+			}
+
 #if DO_CHECK
 			ReaderMap.Emplace(Thread, new FThreadCheckingArchiveProxy(PakReader, Thread));
 #else //DO_CHECK
@@ -5418,27 +5419,27 @@ IFileHandle* FPakPlatformFile::CreatePakFileHandle(const TCHAR* Filename, FPakFi
 {
 	IFileHandle* Result = NULL;
 	bool bNeedsDelete = true;
-	FArchive* PakReader = PakFile->GetSharedReader(LowerLevel);
+	TFunction<FArchive*()> AcquirePakReader = [PakFile, LowerLevelPlatformFile = LowerLevel]() { return PakFile->GetSharedReader(LowerLevelPlatformFile); };
 
 	// Create the handle.
 	if (FileEntry->CompressionMethodIndex != 0 && PakFile->GetInfo().Version >= FPakInfo::PakFile_Version_CompressionEncryption)
 	{
 		if (FileEntry->IsEncrypted())
 		{
-			Result = new FPakFileHandle< FPakCompressedReaderPolicy<FPakSimpleEncryption> >(*PakFile, *FileEntry, PakReader, bNeedsDelete);
+			Result = new FPakFileHandle< FPakCompressedReaderPolicy<FPakSimpleEncryption> >(*PakFile, *FileEntry, AcquirePakReader, bNeedsDelete);
 		}
 		else
 		{
-			Result = new FPakFileHandle< FPakCompressedReaderPolicy<> >(*PakFile, *FileEntry, PakReader, bNeedsDelete);
+			Result = new FPakFileHandle< FPakCompressedReaderPolicy<> >(*PakFile, *FileEntry, AcquirePakReader, bNeedsDelete);
 		}
 	}
 	else if (FileEntry->IsEncrypted())
 	{
-		Result = new FPakFileHandle< FPakReaderPolicy<FPakSimpleEncryption> >(*PakFile, *FileEntry, PakReader, bNeedsDelete);
+		Result = new FPakFileHandle< FPakReaderPolicy<FPakSimpleEncryption> >(*PakFile, *FileEntry, AcquirePakReader, bNeedsDelete);
 	}
 	else
 	{
-		Result = new FPakFileHandle<>(*PakFile, *FileEntry, PakReader, bNeedsDelete);
+		Result = new FPakFileHandle<>(*PakFile, *FileEntry, AcquirePakReader, bNeedsDelete);
 	}
 
 	return Result;
diff --git a/Engine/Source/Runtime/PakFile/Public/IPlatformFilePak.h b/Engine/Source/Runtime/PakFile/Public/IPlatformFilePak.h
--- a/Engine/Source/Runtime/PakFile/Public/IPlatformFilePak.h
+++ b/Engine/Source/Runtime/PakFile/Public/IPlatformFilePak.h
@@ -1093,6 +1093,11 @@ public:
 	}
 };
 
+/**
+ * Typedef for a function that returns an archive to use for accessing an underlying pak file
+ */
+typedef TFunction<FArchive*()> TAcquirePakReaderFunction;
+
 template< typename EncryptionPolicy = FPakNoEncryption >
 class PAKFILE_API FPakReaderPolicy
 {
@@ -1102,14 +1107,14 @@ public:
 	/** Pak file entry for this file. */
 	FPakEntry			PakEntry;
 	/** Pak file archive to read the data from. */
-	FArchive*			PakReader;
+	TAcquirePakReaderFunction AcquirePakReader;
 	/** Offset to the file in pak (including the file header). */
 	int64				OffsetToFile;
 
-	FPakReaderPolicy(const FPakFile& InPakFile,const FPakEntry& InPakEntry,FArchive* InPakReader)
+	FPakReaderPolicy(const FPakFile& InPakFile,const FPakEntry& InPakEntry, TAcquirePakReaderFunction& InAcquirePakReader)
 		: PakFile(InPakFile)
 		, PakEntry(InPakEntry)
-		, PakReader(InPakReader)
+		, AcquirePakReader(InAcquirePakReader)
 	{
 		OffsetToFile = PakEntry.Offset + PakEntry.GetSerializedSize(PakFile.GetInfo().Version);
 	}
@@ -1125,6 +1130,7 @@ public:
 		const constexpr int64 Alignment = (int64)EncryptionPolicy::Alignment;
 		const constexpr int64 AlignmentMask = ~(Alignment - 1);
 		uint8 TempBuffer[Alignment];
+		FArchive* PakReader = AcquirePakReader();
 		if (EncryptionPolicy::AlignReadRequest(DesiredPosition) != DesiredPosition)
 		{
 			int64 Start = DesiredPosition & AlignmentMask;
@@ -1174,6 +1180,21 @@ class PAKFILE_API FPakFileHandle : public IFileHandle
 
 public:
 
+	/**
+	 * Constructs pak file handle to read from pak.
+	 *
+	 * @param InFilename Filename
+	 * @param InPakEntry Entry in the pak file.
+	 * @param InAcquirePakReaderFunction Function that returns the archive to use for serialization. The result of this should not be cached, but reacquired on each serialization operation
+	 */
+	FPakFileHandle(const FPakFile& InPakFile, const FPakEntry& InPakEntry, TAcquirePakReaderFunction& InAcquirePakReaderFunction, bool bIsSharedReader)
+		: bSharedReader(bIsSharedReader)
+		, ReadPos(0)
+		, Reader(InPakFile, InPakEntry, InAcquirePakReaderFunction)
+	{
+		INC_DWORD_STAT(STAT_PakFile_NumOpenHandles);
+	}
+
 	/**
 	 * Constructs pak file handle to read from pak.
 	 *
@@ -1184,7 +1205,7 @@ public:
 	FPakFileHandle(const FPakFile& InPakFile, const FPakEntry& InPakEntry, FArchive* InPakReader, bool bIsSharedReader)
 		: bSharedReader(bIsSharedReader)
 		, ReadPos(0)
-		, Reader(InPakFile, InPakEntry, InPakReader)
+		, Reader(InPakFile, InPakEntry, [InPakReader]() { return InPakReader; })
 	{
 		INC_DWORD_STAT(STAT_PakFile_NumOpenHandles);
 	}
@@ -1196,7 +1217,7 @@ public:
 	{
 		if (!bSharedReader)
 		{
-			delete Reader.PakReader;
+			delete Reader.AcquirePakReader();
 		}
 
 		DEC_DWORD_STAT(STAT_PakFile_NumOpenHandles);
@@ -1228,8 +1249,9 @@ public:
 		if (!Reader.PakEntry.Verified)
 		{
 			FPakEntry FileHeader;
-			Reader.PakReader->Seek(Reader.PakEntry.Offset);
-			FileHeader.Serialize(*Reader.PakReader, Reader.PakFile.GetInfo().Version);
+			FArchive* PakReader = Reader.AcquirePakReader();
+			PakReader->Seek(Reader.PakEntry.Offset);
+			FileHeader.Serialize(*PakReader, Reader.PakFile.GetInfo().Version);
 			if (FPakEntry::VerifyPakEntriesMatch(Reader.PakEntry, FileHeader))
 			{
 				Reader.PakEntry.Verified = true;
