When a shader pak reader is first created for a thread, return the FThreadCheckingArchiveProxy wrapped version rather than the raw archive pointer.

diff --git a/Engine/Source/Runtime/PakFile/Private/IPlatformFilePak.cpp b/Engine/Source/Runtime/PakFile/Private/IPlatformFilePak.cpp
--- a/Engine/Source/Runtime/PakFile/Private/IPlatformFilePak.cpp
+++ b/Engine/Source/Runtime/PakFile/Private/IPlatformFilePak.cpp
@@ -4855,7 +4855,9 @@ FArchive* FPakFile::GetSharedReader(IPlatformFile* LowerLevel)
 			}
 
 #if DO_CHECK
-			ReaderMap.Emplace(Thread, new FThreadCheckingArchiveProxy(PakReader, Thread));
+			FArchive* Proxy = new FThreadCheckingArchiveProxy(PakReader, Thread);
+			ReaderMap.Emplace(Thread, Proxy);
+			PakReader = Proxy;
 #else //DO_CHECK
 			ReaderMap.Emplace(Thread, PakReader);
 #endif //DO_CHECK
