- Fixed conflict from a merge to main

diff --git a/Engine/Source/Runtime/PakFile/Private/IPlatformFilePak.cpp b/Engine/Source/Runtime/PakFile/Private/IPlatformFilePak.cpp
--- a/Engine/Source/Runtime/PakFile/Private/IPlatformFilePak.cpp
+++ b/Engine/Source/Runtime/PakFile/Private/IPlatformFilePak.cpp
@@ -3807,7 +3807,7 @@ IMappedFileHandle* FPakPlatformFile::OpenMapped(const TCHAR* Filename)
 	FPakFile* PakEntry = nullptr;
 	if (FindFileInPakFiles(Filename, &PakEntry, &FileEntry) && PakEntry)
 	{
-		if (FileEntry.CompressionMethod != COMPRESS_None)
+		if (FileEntry.CompressionMethodIndex != 0)
 		{
 			// can't map compressed files
 			return nullptr;
