addition of better hash algorithm for removal of file names from hash

diff --git a/Engine/Source/Runtime/Core/Private/Misc/Fnv.cpp b/Engine/Source/Runtime/Core/Private/Misc/Fnv.cpp
new file mode 100644
--- /dev/null
+++ b/Engine/Source/Runtime/Core/Private/Misc/Fnv.cpp
@@ -0,0 +1,42 @@
+// Copyright 1998-2018 Epic Games, Inc. All Rights Reserved.
+
+#include "Misc/Fnv.h"
+/** Algorithm implementation from http://tools.ietf.org/html/draft-eastlake-fnv-16 */
+
+// generate a 32-bit hash from data in memory of length bytes
+uint32 FFnv::MemFnv32( const void* InData, int32 Length, uint32 InOffset )
+{
+    // constants from above reference
+    static const uint32 Offset = 0x811c9dc5;
+    static const uint32 Prime = 0x01000193;
+    
+    const uint8* __restrict Data = (uint8*)InData;
+    
+    uint32 Fnv = Offset+InOffset;   // this is not strictly correct as the offset should be prime and InOffset could be arbitrary
+    for (; Length; --Length)
+    {
+        Fnv ^= *Data++;
+        Fnv *= Prime;
+    }
+    
+    return Fnv;
+}
+
+// generate a 32-bit hash from data in memory of length bytes
+uint64 FFnv::MemFnv64( const void* InData, int32 Length, uint64 InOffset )
+{
+    // constants from above reference
+    static const uint64 Offset = 0x00000100000001b3;
+    static const uint64 Prime = 0xcbf29ce484222325;
+    
+    const uint8* __restrict Data = (uint8*)InData;
+    
+    uint64 Fnv = Offset+InOffset; // this is not strictly correct as the offset should be prime and InOffset could be arbitrary
+    for (; Length; --Length)
+    {
+        Fnv ^= *Data++;
+        Fnv *= Prime;
+    }
+    
+    return Fnv;
+}
diff --git a/Engine/Source/Runtime/Core/Public/Misc/Fnv.h b/Engine/Source/Runtime/Core/Public/Misc/Fnv.h
new file mode 100644
--- /dev/null
+++ b/Engine/Source/Runtime/Core/Public/Misc/Fnv.h
@@ -0,0 +1,19 @@
+// Copyright 1998-2018 Epic Games, Inc. All Rights Reserved.
+
+#pragma once
+
+#include "CoreTypes.h"
+#include "Misc/AssertionMacros.h"
+#include "Misc/Char.h"
+#include "Misc/CString.h"
+
+
+/** 
+ * FNV hash generation for different types of input data
+ **/
+struct CORE_API FFnv
+{
+	/** generates FNV hash of the memory area */
+	static uint32 MemFnv32( const void* Data, int32 Length, uint32 FNV=0 );
+    static uint64 MemFnv64( const void* Data, int32 Length, uint64 FNV=0 );
+};
diff --git a/Engine/Source/Runtime/PakFile/Private/IPlatformFilePak.cpp b/Engine/Source/Runtime/PakFile/Private/IPlatformFilePak.cpp
--- a/Engine/Source/Runtime/PakFile/Private/IPlatformFilePak.cpp
+++ b/Engine/Source/Runtime/PakFile/Private/IPlatformFilePak.cpp
@@ -26,6 +26,7 @@
 #include "Misc/ConfigCacheIni.h"
 #endif
 #include "ProfilingDebugging/CsvProfiler.h"
+#include "Misc/Fnv.h"
 
 #include "Async/MappedFileHandle.h"
 
@@ -4135,7 +4136,7 @@ FPakFile::FPakFile(const TCHAR* Filename, bool bIsSigned)
 	, bIsValid(false)
 	, bFilenamesRemoved(false)
 	, ChunkID(ParseChunkIDFromFilename(Filename))
-	, MappedFileHandle(nullptr)
+ 	, MappedFileHandle(nullptr)
 {
 	FArchive* Reader = GetSharedReader(NULL);
 	if (Reader)
@@ -4468,7 +4469,7 @@ void FPakFile::UnloadPakEntryFilenames(TArray<FString>* DirectoryRootsToKeep)
 	const int MAX_RETRIES = 10;
 	bool bHasCollision;
 	FilenameStartHash = 0;
-
+    
 	// Allocate the temporary array for hashing filenames. The Memset is to hopefully
 	// silence the Visual Studio static analyzer.
 	TArray<FMiniFileEntry> MiniFileEntries;
@@ -4486,7 +4487,7 @@ void FPakFile::UnloadPakEntryFilenames(TArray<FString>* DirectoryRootsToKeep)
 			for (FPakDirectory::TConstIterator DirectoryIt(It.Value()); DirectoryIt; ++DirectoryIt)
 			{
 				FString FinalFilename = It.Key() / DirectoryIt.Key();
-				uint32 FilenameHash = FCrc::MemCrc32(*FinalFilename.ToLower(), FinalFilename.Len() * sizeof(TCHAR), FilenameStartHash);
+                uint32 FilenameHash = FFnv::MemFnv32(*FinalFilename.ToLower(), FinalFilename.Len() * sizeof(TCHAR), FilenameStartHash);
 				MiniFileEntries[EntryIndex].FilenameHash = FilenameHash;
 				MiniFileEntries[EntryIndex].EntryIndex = DirectoryIt.Value();
 				++EntryIndex;
@@ -5009,7 +5010,7 @@ FPakFile::EFindResult FPakFile::Find(const FString& Filename, FPakEntry* OutEntr
 				++SplitStartPtr;
 				--SplitLen;
 			}
-			uint32 PathHash = FCrc::MemCrc32(SplitStartPtr, SplitLen * sizeof(TCHAR), FilenameStartHash);
+            uint32 PathHash = FFnv::MemFnv32(SplitStartPtr, SplitLen * sizeof(TCHAR), FilenameStartHash);
 
 			// Look it up in our sorted-by-filename-hash array.
 			uint32 PathHashMostSignificantBits = PathHash >> 24;
