Disable SHA1 hashes for pak signing. Still experimental, and not indended to be enabled by default

diff --git a/Engine/Source/Runtime/PakFile/Public/IPlatformFilePak.h b/Engine/Source/Runtime/PakFile/Public/IPlatformFilePak.h
--- a/Engine/Source/Runtime/PakFile/Public/IPlatformFilePak.h
+++ b/Engine/Source/Runtime/PakFile/Public/IPlatformFilePak.h
@@ -26,7 +26,7 @@ DECLARE_DWORD_ACCUMULATOR_STAT_EXTERN(TEXT("Num open pak file handles"), STAT_Pa
 #define PAK_TRACKER 0
 
 // Define the type of a chunk hash. Currently selectable between SHA1 and CRC32.
-#define PAKHASH_USE_CRC	0
+#define PAKHASH_USE_CRC	1
 #if PAKHASH_USE_CRC
 typedef uint32 TPakChunkHash;
 #else
