include a CRC of the compressed data in the fatal error log if decompression fails to help narrow down where in the IO pipline the failure occurred.

diff --git a/Engine/Source/Runtime/PakFile/Private/IPlatformFilePak.cpp b/Engine/Source/Runtime/PakFile/Private/IPlatformFilePak.cpp
--- a/Engine/Source/Runtime/PakFile/Private/IPlatformFilePak.cpp
+++ b/Engine/Source/Runtime/PakFile/Private/IPlatformFilePak.cpp
@@ -3489,7 +3489,7 @@ public:
 
 			if( !FCompression::UncompressMemory(CompressionMethod, Output, Block.ProcessedSize, Block.Raw, Block.DecompressionRawSize) )
 			{
-				UE_LOG( LogPakFile, Fatal, TEXT("Pak Decompression failed. PakFile: %s. EntryOffset: %lld, EntrySize: %lld, CompressionMethod:%s Output:%p  ProcessedSize:%d  Buf:%p  Block.DecompressionRawSize:%d "), *PakFile.ToString(), FileEntry.Offset, FileEntry.Size, *CompressionMethod.ToString(), Output, Block.ProcessedSize, Block.Raw, Block.DecompressionRawSize );
+				UE_LOG( LogPakFile, Fatal, TEXT("Pak Decompression failed. PakFile: %s. EntryOffset: %lld, EntrySize: %lld, CompressionMethod:%s Output:%p  ProcessedSize:%d  Buf:%p  Block.DecompressionRawSize:%d  Crc32:%u"), *PakFile.ToString(), FileEntry.Offset, FileEntry.Size, *CompressionMethod.ToString(), Output, Block.ProcessedSize, Block.Raw, Block.DecompressionRawSize, FCrc::MemCrc32( Block.Raw, Block.DecompressionRawSize ) );
 			}
 			FMemory::Free(Block.Raw);
 			Block.Raw = nullptr;
