Make pak platform file wnbind the RegisterEncryptionKey delegate when it is destroyed

diff --git a/Engine/Source/Runtime/PakFile/Private/IPlatformFilePak.cpp b/Engine/Source/Runtime/PakFile/Private/IPlatformFilePak.cpp
--- a/Engine/Source/Runtime/PakFile/Private/IPlatformFilePak.cpp
+++ b/Engine/Source/Runtime/PakFile/Private/IPlatformFilePak.cpp
@@ -5401,6 +5401,7 @@ FPakPlatformFile::FPakPlatformFile()
 
 FPakPlatformFile::~FPakPlatformFile()
 {
+	FCoreDelegates::GetRegisterEncryptionKeyDelegate().Unbind();
 	FCoreDelegates::OnMountPak.Unbind();
 	FCoreDelegates::OnUnmountPak.Unbind();
 
