Clean up missing packages list when we mount pak files.

diff --git a/Engine/Source/Runtime/PakFile/Private/IPlatformFilePak.cpp b/Engine/Source/Runtime/PakFile/Private/IPlatformFilePak.cpp
--- a/Engine/Source/Runtime/PakFile/Private/IPlatformFilePak.cpp
+++ b/Engine/Source/Runtime/PakFile/Private/IPlatformFilePak.cpp
@@ -29,6 +29,7 @@
 #include "ProfilingDebugging/CsvProfiler.h"
 #include "Misc/Fnv.h"
 
+
 #include "Async/MappedFileHandle.h"
 
 DEFINE_LOG_CATEGORY(LogPakFile);
@@ -185,7 +186,7 @@ FPakMasterSignatureTableCheckFailureHandler& FPakPlatformFile::GetPakMasterSigna
 	return Delegate;
 }
 
-void FPakPlatformFile::GetFilenamesInChunk(const FString& InPakFilename, const TArray<int32>& InChunkIDs, TArray<FString>& OutFileList)
+void FPakPlatformFile::GetFilenamesInChunk(const FString& InPakFilename, const TArray<int32>& InChunkIDs, TArray<FString>& OutFileList) 
 {
 	TArray<FPakListEntry> Paks;
 	GetMountedPaks(Paks);
@@ -200,6 +201,21 @@ void FPakPlatformFile::GetFilenamesInChunk(const FString& InPakFilename, const T
 	}
 }
 
+void FPakPlatformFile::GetFilenamesInPakFile(const FString& InPakFilename, TArray<FString>& OutFileList)
+{
+	TArray<FPakListEntry> Paks;
+	GetMountedPaks(Paks);
+
+	for (const FPakListEntry& Pak : Paks)
+	{
+		if (Pak.PakFile && Pak.PakFile->GetFilename() == InPakFilename)
+		{
+			Pak.PakFile->GetFilenames(OutFileList);
+			break;
+		}
+	}
+}
+
 #define USE_PAK_PRECACHE (!IS_PROGRAM && !WITH_EDITOR) // you can turn this off to use the async IO stuff without the precache
 
 /**
@@ -5114,6 +5130,20 @@ public:
 };
 #endif //DO_CHECK
 
+
+void FPakFile::GetFilenames(TArray<FString>& OutFileList) const
+{
+	for (const TMap<FString, FPakDirectory>::ElementType& DirectoryElement : Index)
+	{
+		const  FPakDirectory& Directory = DirectoryElement.Value;
+		for (const FPakDirectory::ElementType& FileElement : Directory)
+		{
+			OutFileList.Add(MountPoint / DirectoryElement.Key / FileElement.Key);
+		}
+	}
+}
+
+
 void FPakFile::GetFilenamesInChunk(const TArray<int32>& InChunkIDs, TArray<FString>& OutFileList)
 {
 	TSet<int32> OverlappingEntries;
diff --git a/Engine/Source/Runtime/PakFile/Public/IPlatformFilePak.h b/Engine/Source/Runtime/PakFile/Public/IPlatformFilePak.h
--- a/Engine/Source/Runtime/PakFile/Public/IPlatformFilePak.h
+++ b/Engine/Source/Runtime/PakFile/Public/IPlatformFilePak.h
@@ -659,6 +659,8 @@ public:
 		return Files.Num();
 	}
 
+	void GetFilenames(TArray<FString>& OutFileList) const;
+
 	void GetFilenamesInChunk(const TArray<int32>& InChunkIDs, TArray<FString>& OutFileList);
 
 	/**
@@ -2355,6 +2357,7 @@ public:
 
 	// Get a list of which files live in a given chunk
 	void GetFilenamesInChunk(const FString& InPakFilename, const TArray<int32>& InChunkIDs, TArray<FString>& OutFileList);
+	void GetFilenamesInPakFile(const FString& InPakFilename, TArray<FString>& OutFileList) ;
 
 	void UnloadPakEntryFilenames(TArray<FString>* DirectoryRootsToKeep = nullptr);
 	void ShrinkPakEntriesMemoryUsage();
