Clean up missing packages list when we mount pak files.

diff --git a/Engine/Source/Runtime/Core/Private/Misc/CoreDelegates.cpp b/Engine/Source/Runtime/Core/Private/Misc/CoreDelegates.cpp
--- a/Engine/Source/Runtime/Core/Private/Misc/CoreDelegates.cpp
+++ b/Engine/Source/Runtime/Core/Private/Misc/CoreDelegates.cpp
@@ -32,6 +32,7 @@ FCoreDelegates::FOnUnmountPak FCoreDelegates::OnUnmountPak;
 FCoreDelegates::FOnOptimizeMemoryUsageForMountedPaks FCoreDelegates::OnOptimizeMemoryUsageForMountedPaks;
 
 FCoreDelegates::FPakFileMountedDelegate FCoreDelegates::PakFileMountedCallback;
+FCoreDelegates::FNewFileAddedDelegate FCoreDelegates::NewFileAddedDelegate;
 FCoreDelegates::FNoPakFilesMountedDelegate FCoreDelegates::NoPakFilesMountedDelegate;
 FCoreDelegates::FOnFileOpenedForReadFromPakFile FCoreDelegates::OnFileOpenedForReadFromPakFile;
 
diff --git a/Engine/Source/Runtime/Core/Public/Misc/CoreDelegates.h b/Engine/Source/Runtime/Core/Public/Misc/CoreDelegates.h
--- a/Engine/Source/Runtime/Core/Public/Misc/CoreDelegates.h
+++ b/Engine/Source/Runtime/Core/Public/Misc/CoreDelegates.h
@@ -161,6 +161,10 @@ public:
 	// After a pakfile is mounted this callback is called for each new file
 	static FPakFileMountedDelegate PakFileMountedCallback;
 
+	// After a file is added this is called
+	DECLARE_MULTICAST_DELEGATE_OneParam(FNewFileAddedDelegate, const FString&);
+	static FNewFileAddedDelegate NewFileAddedDelegate;
+
 	// After an attempt to mount all pak files, but none wre found, this is called
 	static FNoPakFilesMountedDelegate NoPakFilesMountedDelegate;
 
diff --git a/Engine/Source/Runtime/CoreUObject/Private/UObject/LinkerLoad.cpp b/Engine/Source/Runtime/CoreUObject/Private/UObject/LinkerLoad.cpp
--- a/Engine/Source/Runtime/CoreUObject/Private/UObject/LinkerLoad.cpp
+++ b/Engine/Source/Runtime/CoreUObject/Private/UObject/LinkerLoad.cpp
@@ -5121,6 +5121,20 @@ bool FLinkerLoad::RemoveKnownMissingPackage(FName PackageName)
 	return FCoreRedirects::RemoveKnownMissing(ECoreRedirectFlags::Type_Package, FCoreRedirectObjectName(NAME_None, NAME_None, PackageName));
 }
 
+void FLinkerLoad::OnNewFileAdded(const FString& Filename)
+{
+	FString PackageName;
+	if (FPackageName::TryConvertFilenameToLongPackageName(Filename, PackageName))
+	{
+		FName PackageFName(*PackageName);
+		if (FLinkerLoad::IsKnownMissingPackage(PackageFName))
+		{
+			FLinkerLoad::RemoveKnownMissingPackage(PackageFName);
+		}
+	}
+}
+
+
 void FLinkerLoad::AddGameNameRedirect(const FName OldName, const FName NewName)
 {
 	TArray<FCoreRedirect> NewRedirects;
diff --git a/Engine/Source/Runtime/CoreUObject/Private/UObject/Obj.cpp b/Engine/Source/Runtime/CoreUObject/Private/UObject/Obj.cpp
--- a/Engine/Source/Runtime/CoreUObject/Private/UObject/Obj.cpp
+++ b/Engine/Source/Runtime/CoreUObject/Private/UObject/Obj.cpp
@@ -4296,7 +4296,7 @@ void InitUObject()
 	};
 	FModuleManager::Get().IsPackageLoadedCallback().BindStatic(Local::IsPackageLoaded);
 	
-
+	FCoreDelegates::NewFileAddedDelegate.AddStatic(FLinkerLoad::OnNewFileAdded);
 
 #if WITH_EDITOR
 PRAGMA_DISABLE_DEPRECATION_WARNINGS
diff --git a/Engine/Source/Runtime/CoreUObject/Public/UObject/LinkerLoad.h b/Engine/Source/Runtime/CoreUObject/Public/UObject/LinkerLoad.h
--- a/Engine/Source/Runtime/CoreUObject/Public/UObject/LinkerLoad.h
+++ b/Engine/Source/Runtime/CoreUObject/Public/UObject/LinkerLoad.h
@@ -244,6 +244,11 @@ public:
 	 */
 	COREUOBJECT_API static bool RemoveKnownMissingPackage(FName PackageName);
 
+	/**
+	 * 
+	 */
+	COREUOBJECT_API static void OnNewFileAdded(const FString& Filename);
+
 	/** 
 	 * Checks if the linker has any objects in the export table that require loading.
 	 */
diff --git a/Engine/Source/Runtime/PakFile/Private/IPlatformFilePak.cpp b/Engine/Source/Runtime/PakFile/Private/IPlatformFilePak.cpp
--- a/Engine/Source/Runtime/PakFile/Private/IPlatformFilePak.cpp
+++ b/Engine/Source/Runtime/PakFile/Private/IPlatformFilePak.cpp
@@ -185,7 +185,7 @@ FPakMasterSignatureTableCheckFailureHandler& FPakPlatformFile::GetPakMasterSigna
 	return Delegate;
 }
 
-void FPakPlatformFile::GetFilenamesInChunk(const FString& InPakFilename, const TArray<int32>& InChunkIDs, TArray<FString>& OutFileList)
+void FPakPlatformFile::GetFilenamesInChunk(const FString& InPakFilename, const TArray<int32>& InChunkIDs, TArray<FString>& OutFileList) 
 {
 	TArray<FPakListEntry> Paks;
 	GetMountedPaks(Paks);
@@ -200,6 +200,21 @@ void FPakPlatformFile::GetFilenamesInChunk(const FString& InPakFilename, const T
 	}
 }
 
+void FPakPlatformFile::GetFilenamesInPakFile(const FString& InPakFilename, TArray<FString>& OutFileList)
+{
+	TArray<FPakListEntry> Paks;
+	GetMountedPaks(Paks);
+
+	for (const FPakListEntry& Pak : Paks)
+	{
+		if (Pak.PakFile && Pak.PakFile->GetFilename() == InPakFilename)
+		{
+			Pak.PakFile->GetFilenames(OutFileList);
+			break;
+		}
+	}
+}
+
 #define USE_PAK_PRECACHE (!IS_PROGRAM && !WITH_EDITOR) // you can turn this off to use the async IO stuff without the precache
 
 /**
@@ -5114,6 +5129,20 @@ public:
 };
 #endif //DO_CHECK
 
+
+void FPakFile::GetFilenames(TArray<FString>& OutFileList) const
+{
+	for (const TMap<FString, FPakDirectory>::ElementType& DirectoryElement : Index)
+	{
+		const  FPakDirectory& Directory = DirectoryElement.Value;
+		for (const FPakDirectory::ElementType& FileElement : Directory)
+		{
+			OutFileList.Add(MountPoint / DirectoryElement.Key / FileElement.Key);
+		}
+	}
+}
+
+
 void FPakFile::GetFilenamesInChunk(const TArray<int32>& InChunkIDs, TArray<FString>& OutFileList)
 {
 	TSet<int32> OverlappingEntries;
@@ -5733,6 +5762,24 @@ bool FPakPlatformFile::Mount(const TCHAR* InPakFilename, uint32 PakOrder, const
 		{
 			UE_LOG(LogPakFile, Warning, TEXT("Failed to mount pak \"%s\", pak is invalid."), InPakFilename);
 		}
+
+		if (bSuccess)
+		{
+			if (FCoreDelegates::PakFileMountedCallback.IsBound())
+			{
+				// uncomment this when we have time to test 
+				// FCoreDelegates::PakFileMountedCallback.Broadcast(InPakFilename);
+			}
+			if (FCoreDelegates::NewFileAddedDelegate.IsBound())
+			{
+				TArray<FString> Filenames;
+				Pak->GetFilenames(Filenames);
+				for (const FString& Filename : Filenames)
+				{
+					FCoreDelegates::NewFileAddedDelegate.Broadcast(Filename);
+				}
+			}
+		}
 	}
 	else
 	{
diff --git a/Engine/Source/Runtime/PakFile/Public/IPlatformFilePak.h b/Engine/Source/Runtime/PakFile/Public/IPlatformFilePak.h
--- a/Engine/Source/Runtime/PakFile/Public/IPlatformFilePak.h
+++ b/Engine/Source/Runtime/PakFile/Public/IPlatformFilePak.h
@@ -659,6 +659,8 @@ public:
 		return Files.Num();
 	}
 
+	void GetFilenames(TArray<FString>& OutFileList) const;
+
 	void GetFilenamesInChunk(const TArray<int32>& InChunkIDs, TArray<FString>& OutFileList);
 
 	/**
@@ -2355,6 +2357,7 @@ public:
 
 	// Get a list of which files live in a given chunk
 	void GetFilenamesInChunk(const FString& InPakFilename, const TArray<int32>& InChunkIDs, TArray<FString>& OutFileList);
+	void GetFilenamesInPakFile(const FString& InPakFilename, TArray<FString>& OutFileList) ;
 
 	void UnloadPakEntryFilenames(TArray<FString>* DirectoryRootsToKeep = nullptr);
 	void ShrinkPakEntriesMemoryUsage();
