If no pak files are found when IPlatformFilePak mounts paks, fallback to mounting all paks.

diff --git a/Engine/Source/Runtime/PakFile/Private/IPlatformFilePak.cpp b/Engine/Source/Runtime/PakFile/Private/IPlatformFilePak.cpp
--- a/Engine/Source/Runtime/PakFile/Private/IPlatformFilePak.cpp
+++ b/Engine/Source/Runtime/PakFile/Private/IPlatformFilePak.cpp
@@ -6020,6 +6020,13 @@ int32 FPakPlatformFile::MountAllPakFiles(const TArray<FString>& PakFolders, cons
 	{
 		TArray<FString> FoundPakFiles;
 		FindAllPakFiles(LowerLevel, PakFolders, WildCard, FoundPakFiles);
+
+		// HACK: If no pak files are found with the wildcard, fallback to mounting everything.
+		if (FoundPakFiles.Num() == 0)
+		{
+			FindAllPakFiles(LowerLevel, PakFolders, ALL_PAKS_WILDCARD, FoundPakFiles);
+		}
+
 		// Sort in descending order.
 		FoundPakFiles.Sort(TGreater<FString>());
 		// Mount all found pak files
