When a chunk is installed, mount optional pak files and other pak files mapped to this chunk.

diff --git a/Engine/Source/Runtime/PakFile/Private/IPlatformFilePak.cpp b/Engine/Source/Runtime/PakFile/Private/IPlatformFilePak.cpp
--- a/Engine/Source/Runtime/PakFile/Private/IPlatformFilePak.cpp
+++ b/Engine/Source/Runtime/PakFile/Private/IPlatformFilePak.cpp
@@ -5904,8 +5904,7 @@ bool FPakPlatformFile::Mount(const TCHAR* InPakFilename, uint32 PakOrder, const
 		{
 			if (FCoreDelegates::PakFileMountedCallback.IsBound())
 			{
-				// uncomment this when we have time to test 
-				// FCoreDelegates::PakFileMountedCallback.Broadcast(InPakFilename);
+				FCoreDelegates::PakFileMountedCallback.Broadcast(InPakFilename);
 			}
 			if (FCoreDelegates::NewFileAddedDelegate.IsBound())
 			{
