PakFile: Change a "verify" on the required encryption key being available to a fatal as it is never going to be a recoverable event and we don't want to continue if verify is compiled out

diff --git a/Engine/Source/Runtime/PakFile/Private/IPlatformFilePak.cpp b/Engine/Source/Runtime/PakFile/Private/IPlatformFilePak.cpp
--- a/Engine/Source/Runtime/PakFile/Private/IPlatformFilePak.cpp
+++ b/Engine/Source/Runtime/PakFile/Private/IPlatformFilePak.cpp
@@ -227,7 +227,10 @@ void FPakPlatformFile::GetPakEncryptionKey(FAES::FAESKey& OutKey, const FGuid& I
 
 	if (InEncryptionKeyGuid.IsValid())
 	{
-		verify(GetRegisteredEncryptionKeys().GetKey(InEncryptionKeyGuid, OutKey));
+		if (!GetRegisteredEncryptionKeys().GetKey(InEncryptionKeyGuid, OutKey))
+		{
+			UE_LOG(LogPakFile, Fatal, TEXT("Failed to find requested encryption key %s"), *InEncryptionKeyGuid.ToString());
+		}
 	}
 	else
 	{
