Allow a pak file to be made using the subset of another pak file.
Allow more control over the pak cache, cache per pak, indvidual caches for all paks.
Increase the number of cached blocks for XboxOne.
Enable cache per pak for XboxOne.

gil.gribb, graeme.thornton

diff --git a/Engine/Source/Runtime/Core/Public/GenericPlatform/GenericPlatformFile.h b/Engine/Source/Runtime/Core/Public/GenericPlatform/GenericPlatformFile.h
--- a/Engine/Source/Runtime/Core/Public/GenericPlatform/GenericPlatformFile.h
+++ b/Engine/Source/Runtime/Core/Public/GenericPlatform/GenericPlatformFile.h
@@ -233,6 +233,11 @@ public:
 	 */
 	virtual void		InitializeAfterSetActive() { }
 
+	/**
+	 * Build an in memory unique pak file from a subset of files in this pak file
+	 */
+	virtual void		MakeUniquePakFilesForTheseFiles(TArray<TArray<FString>> InFiles) { }
+
 	/**
 	* Performs initialization of the platform file after the new async IO has been enabled
 	*/
diff --git a/Engine/Source/Runtime/PakFile/Private/IPlatformFilePak.cpp b/Engine/Source/Runtime/PakFile/Private/IPlatformFilePak.cpp
--- a/Engine/Source/Runtime/PakFile/Private/IPlatformFilePak.cpp
+++ b/Engine/Source/Runtime/PakFile/Private/IPlatformFilePak.cpp
@@ -23,6 +23,7 @@
 #include "Misc/Base64.h"
 #include "Stats/StatsMisc.h"
 #include "HAL/PlatformFilemanager.h"
+#include "HAL/ThreadHeartBeat.h"
 #if !(IS_PROGRAM || WITH_EDITOR)
 #include "Misc/ConfigCacheIni.h"
 #endif
@@ -238,21 +239,13 @@ void FPakPlatformFile::GetPakEncryptionKey(FAES::FAESKey& OutKey, const FGuid& I
 	}
 }
 
-TMap<FName, TSharedPtr<const struct FPakSignatureFile, ESPMode::ThreadSafe>> FPakPlatformFile::PakSignatureFileCache;
-FCriticalSection FPakPlatformFile::PakSignatureFileCacheLock;
-TSharedPtr<const struct FPakSignatureFile, ESPMode::ThreadSafe> FPakPlatformFile::GetPakSignatureFile(const TCHAR* InFilename)
+TSharedPtr<FRSA::FKey, ESPMode::ThreadSafe> FPakPlatformFile::GetPakSigningKey()
 {
-	FScopeLock Lock(&PakSignatureFileCacheLock);
+	static TSharedPtr<FRSA::FKey, ESPMode::ThreadSafe> Key;
+	static FCriticalSection Lock;
+	Lock.Lock();
 
-	FName FilenameFName(InFilename);
-	if (const TSharedPtr<const struct FPakSignatureFile, ESPMode::ThreadSafe>* SignaturesFile = PakSignatureFileCache.Find(FilenameFName))
-	{
-		return *SignaturesFile;
-	}
-
-	static FRSAKeyHandle PublicKey = InvalidRSAKeyHandle;
-	static bool bInitializedPublicKey = false;
-	if (!bInitializedPublicKey)
+	if (!Key.IsValid())
 	{
 		FCoreDelegates::FPakSigningKeysDelegate& Delegate = FCoreDelegates::GetPakSigningKeysDelegate();
 		if (Delegate.IsBound())
@@ -260,27 +253,35 @@ TSharedPtr<const struct FPakSignatureFile, ESPMode::ThreadSafe> FPakPlatformFile
 			TArray<uint8> Exponent;
 			TArray<uint8> Modulus;
 			Delegate.Execute(Exponent, Modulus);
-			PublicKey = FRSA::CreateKey(Exponent, TArray<uint8>(), Modulus);
+			Key = FRSA::CreateKey(Exponent, TArray<uint8>(), Modulus);
 		}
-		bInitializedPublicKey = true;
 	}
+	
+	Lock.Unlock();
+	return Key;
+}
 
-	TSharedPtr<FPakSignatureFile, ESPMode::ThreadSafe> SignaturesFile;
-
-	if (PublicKey != InvalidRSAKeyHandle)
+TMap<FName, TSharedPtr<const struct FPakSignatureFile, ESPMode::ThreadSafe>> FPakPlatformFile::PakSignatureFileCache;
+FCriticalSection FPakPlatformFile::PakSignatureFileCacheLock;
+TSharedPtr<const struct FPakSignatureFile, ESPMode::ThreadSafe> FPakPlatformFile::GetPakSignatureFile(const TCHAR* InFilename)
+{
+	FScopeLock Lock(&PakSignatureFileCacheLock);
+	FName FilenameFName(InFilename);
+	if (const TSharedPtr<const struct FPakSignatureFile, ESPMode::ThreadSafe>* SignaturesFile = PakSignatureFileCache.Find(FilenameFName))
 	{
+		return *SignaturesFile;
+	}
+
+	TSharedPtr<FPakSignatureFile, ESPMode::ThreadSafe> SignaturesFile = MakeShared<FPakSignatureFile, ESPMode::ThreadSafe>();
 	FString SignaturesFilename = FPaths::ChangeExtension(InFilename, TEXT("sig"));
 	FArchive* Reader = IFileManager::Get().CreateFileReader(*SignaturesFilename);
 	if (Reader != nullptr)
 	{
-			SignaturesFile = MakeShared<FPakSignatureFile, ESPMode::ThreadSafe>();
 		SignaturesFile->Serialize(*Reader);
 		delete Reader;
 
-			if (!SignaturesFile->DecryptSignatureAndValidate(PublicKey, InFilename))
+		if (!SignaturesFile->DecryptSignatureAndValidate(InFilename))
 		{
-				// We don't need to act on this failure as the decrypt function will already have dumped out log messages
-				// and fired the signature check fail handler
 			SignaturesFile.Reset();
 		}
 
@@ -291,7 +292,6 @@ TSharedPtr<const struct FPakSignatureFile, ESPMode::ThreadSafe> FPakPlatformFile
 		UE_LOG(LogPakFile, Warning, TEXT("Couldn't find pak signature file '%s'"), InFilename);
 		FPakPlatformFile::GetPakMasterSignatureTableCheckFailureHandler().Broadcast(InFilename);
 	}
-	}
 
 	return SignaturesFile;
 }
@@ -338,7 +338,7 @@ static FAutoConsoleVariableRef CVar_ForceDecompressionFails(
 
 #if USE_PAK_PRECACHE
 #include "Async/TaskGraphInterfaces.h"
-#define PAK_CACHE_GRANULARITY (64*1024)
+#define PAK_CACHE_GRANULARITY (1024*64)
 static_assert((PAK_CACHE_GRANULARITY % FPakInfo::MaxChunkDataSize) == 0, "PAK_CACHE_GRANULARITY must be set to a multiple of FPakInfo::MaxChunkDataSize");
 #define PAK_CACHE_MAX_REQUESTS (8)
 #define PAK_CACHE_MAX_PRIORITY_DIFFERENCE_MERGE (AIOP_Normal - AIOP_MIN)
@@ -358,6 +358,28 @@ static FAutoConsoleVariableRef CVar_Enable(
 	TEXT("If > 0, then enable the pak cache.")
 );
 
+int32 GPakCache_CachePerPakFile = 0;
+static FAutoConsoleVariableRef CVar_CachePerPakFile(
+	TEXT("pakcache.CachePerPakFile"),
+	GPakCache_CachePerPakFile,
+	TEXT("if > 0, then each pak file will have it's own cache")
+);
+
+int32 GPakCache_UseNewTrim = 0;
+static FAutoConsoleVariableRef CVar_UseNewTrim(
+	TEXT("pakcache.UseNewTrim"),
+	GPakCache_UseNewTrim,
+	TEXT("if > 0, then we'll use a round robin per pak file trim")
+);
+
+int32 GPakCache_MaxBlockMemory = 128;
+static FAutoConsoleVariableRef CVar_MaxBlockMemory(
+	TEXT("pakcache.MaxBlockMemory"),
+	GPakCache_MaxBlockMemory,
+	TEXT("A soft memory budget in MB for the max memory used for precaching, that we'll try and adhere to ")
+);
+
+
 int32 GPakCache_MaxRequestsToLowerLevel = 2;
 static FAutoConsoleVariableRef CVar_MaxRequestsToLowerLevel(
 	TEXT("pakcache.MaxRequestsToLowerLevel"),
@@ -1124,6 +1146,7 @@ class FPakPrecacher
 	struct FPakData
 	{
 		IAsyncReadFileHandle* Handle;
+		FPakFile* ActualPakFile;
 		int64 TotalSize;
 		uint64 MaxNode;
 		uint32 StartShift;
@@ -1136,8 +1159,9 @@ class FPakPrecacher
 
 		TSharedPtr<const FPakSignatureFile, ESPMode::ThreadSafe> Signatures;
 
-		FPakData(IAsyncReadFileHandle* InHandle, FName InName, int64 InTotalSize)
+		FPakData(FPakFile* InActualPakFile, IAsyncReadFileHandle* InHandle, FName InName, int64 InTotalSize)
 			: Handle(InHandle)
+			, ActualPakFile(InActualPakFile)
 			, TotalSize(InTotalSize)
 			, StartShift(0)
 			, MaxShift(0)
@@ -1193,14 +1217,14 @@ class FPakPrecacher
 			}
 		}
 	};
-	TMap<FName, uint16> CachedPaks;
+	TMap<FPakFile*, uint16> CachedPaks;
 	TArray<FPakData> CachedPakData;
 
 	TIntervalTreeAllocator<FPakInRequest> InRequestAllocator;
 	TIntervalTreeAllocator<FCacheBlock> CacheBlockAllocator;
 	TMap<uint64, TIntervalTreeIndex> OutstandingRequests;
 
-	TArray<FJoinedOffsetAndPakIndex> OffsetAndPakIndexOfSavedBlocked;
+	TArray < TArray<FJoinedOffsetAndPakIndex>> OffsetAndPakIndexOfSavedBlocked;
 
 	struct FRequestToLower
 	{
@@ -1224,17 +1248,16 @@ class FPakPrecacher
 	uint32 Loads;
 	uint32 Frees;
 	uint64 LoadSize;
+	FRSA::TKeyPtr SigningKey;
 	EAsyncIOPriorityAndFlags AsyncMinPriority;
 	FCriticalSection SetAsyncMinimumPriorityScopeLock;
-	bool bEnableSignatureChecks;
-
 public:
 
-	static void Init(IPlatformFile* InLowerLevel, bool bInEnableSignatureChecks)
+	static void Init(IPlatformFile* InLowerLevel, FRSA::TKeyPtr InSigningKey)
 	{
 		if (!PakPrecacherSingleton)
 		{
-			verify(!FPlatformAtomics::InterlockedCompareExchangePointer((void**)&PakPrecacherSingleton, new FPakPrecacher(InLowerLevel, bInEnableSignatureChecks), nullptr));
+			verify(!FPlatformAtomics::InterlockedCompareExchangePointer((void**)&PakPrecacherSingleton, new FPakPrecacher(InLowerLevel, InSigningKey), nullptr));
 		}
 		check(PakPrecacherSingleton);
 	}
@@ -1270,7 +1293,7 @@ public:
 		return *PakPrecacherSingleton;
 	}
 
-	FPakPrecacher(IPlatformFile* InLowerLevel, bool bInEnableSignatureChecks)
+	FPakPrecacher(IPlatformFile* InLowerLevel, FRSA::TKeyPtr InSigningKey)
 		: LowerLevel(InLowerLevel)
 		, LastReadRequest(0)
 		, NextUniqueID(1)
@@ -1280,8 +1303,8 @@ public:
 		, Loads(0)
 		, Frees(0)
 		, LoadSize(0)
+		, SigningKey(InSigningKey)
 		, AsyncMinPriority(AIOP_MIN)
-		, bEnableSignatureChecks(bInEnableSignatureChecks)
 	{
 		check(LowerLevel && FPlatformProcess::SupportsMultithreading());
 		GPakCache_MaxRequestsToLowerLevel = FMath::Max(FMath::Min(FPlatformMisc::NumberOfIOWorkerThreadsToSpawn(), GPakCache_MaxRequestsToLowerLevel), 1);
@@ -1302,9 +1325,10 @@ public:
 		return LowerLevel;
 	}
 
-	uint16* RegisterPakFile(FName File, int64 PakFileSize)
+	uint16* RegisterPakFile(FPakFile* InActualPakFile, FName File, int64 PakFileSize)
 	{
-		uint16* PakIndexPtr = CachedPaks.Find(File);
+		uint16* PakIndexPtr = CachedPaks.Find(InActualPakFile);
+
 		if (!PakIndexPtr)
 		{
 			FString PakFilename = File.ToString();
@@ -1314,22 +1338,55 @@ public:
 			{
 				return nullptr;
 			}
-			CachedPakData.Add(FPakData(Handle, File, PakFileSize));
-			PakIndexPtr = &CachedPaks.Add(File, CachedPakData.Num() - 1);
+			CachedPakData.Add(FPakData(InActualPakFile, Handle, File, PakFileSize));
+			PakIndexPtr = &CachedPaks.Add(InActualPakFile, CachedPakData.Num() - 1);
+			FPakData& Pak = CachedPakData[*PakIndexPtr];
+
+
+			if (OffsetAndPakIndexOfSavedBlocked.Num() == 0)
+			{
+				// the 1st cache must exist and is used by all sharing pak files
+				OffsetAndPakIndexOfSavedBlocked.AddDefaulted(1);
+			}
+
+			static bool bFirst = true;
+			if (bFirst)
+			{
+				if (FParse::Param(FCommandLine::Get(), TEXT("CachePerPak")))
+				{
+					GPakCache_CachePerPakFile = 1;
+				}
+
+				if (FParse::Param(FCommandLine::Get(), TEXT("NewTrimCache")))
+				{
+					GPakCache_UseNewTrim = 1;
+				}
+				FParse::Value(FCommandLine::Get(), TEXT("PakCacheMaxBlockMemory="), GPakCache_MaxBlockMemory);
+				bFirst = false;
+			}
+
+			if (Pak.ActualPakFile->GetCacheType() == FPakFile::ECacheType::Individual || GPakCache_CachePerPakFile != 0 )
+			{
+				Pak.ActualPakFile->SetCacheIndex(OffsetAndPakIndexOfSavedBlocked.Num());
+				OffsetAndPakIndexOfSavedBlocked.AddDefaulted(1);
+			}
+			else
+			{
+				Pak.ActualPakFile->SetCacheIndex(0);
+			}
+
 			UE_LOG(LogPakFile, Log, TEXT("New pak file %s added to pak precacher."), *PakFilename);
 
-			FPakData& Pak = CachedPakData[*PakIndexPtr];
-			
+			if (SigningKey.IsValid())
+			{
 				// Load signature data
 				Pak.Signatures = FPakPlatformFile::GetPakSignatureFile(*PakFilename);
 				
-			if (Pak.Signatures.IsValid())
-			{
 				// We should never get here unless the signature file exists and is validated. The original FPakFile creation
 				// on the main thread would have failed and the pak would never have been mounted otherwise, and then we would
 				// never have issued read requests to the pak precacher.
 				check(Pak.Signatures);
-
+				
 				// Check that we have the correct match between signature and pre-cache granularity
 				int64 NumPakChunks = Align(PakFileSize, FPakInfo::MaxChunkDataSize) / FPakInfo::MaxChunkDataSize;
 				ensure(NumPakChunks == Pak.Signatures->ChunkHashes.Num());
@@ -1607,44 +1664,135 @@ private: // below here we assume CachedFilesScopeLock until we get to the next s
 		RequestCounter.Decrement();
 		InRequestAllocator.Free(Index);
 	}
-	void TrimCache(bool bDiscardAll = false)
+
+	void TrimCache(bool bDiscardAll = false, uint16 StartPakIndex = 65535)
 	{
-		// CachedFilesScopeLock is locked
-		int32 NumToKeep = bDiscardAll ? 0 : GPakCache_NumUnreferencedBlocksToCache;
-		int32 NumToRemove = FMath::Max<int32>(0, OffsetAndPakIndexOfSavedBlocked.Num() - NumToKeep);
-		if (NumToRemove)
+
+		if (GPakCache_UseNewTrim)
 		{
-			for (int32 Index = 0; Index < NumToRemove; Index++)
+			StartPakIndex = 0;
+			uint16 EndPakIndex = CachedPakData.Num();
+
+			// TODO: remove this array, add a bool to the cache object
+			TArray<bool> CacheVisitedAlready;
+			CacheVisitedAlready.AddDefaulted(OffsetAndPakIndexOfSavedBlocked.Num());
+
+			int64 MemoryBudget = GPakCache_MaxBlockMemory * (1024 * 1024);
+
+
+			while (BlockMemory > MemoryBudget)
 			{
-				FJoinedOffsetAndPakIndex OffsetAndPakIndex = OffsetAndPakIndexOfSavedBlocked[Index];
-				uint16 PakIndex = GetRequestPakIndex(OffsetAndPakIndex);
-				int64 Offset = GetRequestOffset(OffsetAndPakIndex);
-				FPakData& Pak = CachedPakData[PakIndex];
-				MaybeRemoveOverlappingNodesInIntervalTree<FCacheBlock>(
-					&Pak.CacheBlocks[(int32)EBlockStatus::Complete],
-					CacheBlockAllocator,
-					Offset,
-					Offset,
-					0,
-					Pak.MaxNode,
-					Pak.StartShift,
-					Pak.MaxShift,
-					[this](TIntervalTreeIndex BlockIndex) -> bool
+				for (int i = 0; i < CacheVisitedAlready.Num(); i++)
+				{
+					CacheVisitedAlready[i] = false;
+				}
+				// if we iterate over all the pak files and caches and can't remove something then we'll break out of the while.
+				bool NoneToRemove = true;
+				// CachedFilesScopeLock is locked
+				for (uint16 RealPakIndex = StartPakIndex; RealPakIndex < EndPakIndex; RealPakIndex++)
 				{
-					FCacheBlock &Block = CacheBlockAllocator.Get(BlockIndex);
-					if (!Block.InRequestRefCount)
+					int32 CacheIndex = CachedPakData[RealPakIndex].ActualPakFile->GetCacheIndex();
+					if (CacheVisitedAlready[CacheIndex] == true)
+						continue;
+					CacheVisitedAlready[CacheIndex] = true;
+
+					int32 NumToKeep = bDiscardAll ? 0 : GPakCache_NumUnreferencedBlocksToCache;
+					int32 NumToRemove = FMath::Max<int32>(0, OffsetAndPakIndexOfSavedBlocked[CacheIndex].Num() - NumToKeep);
+					if (!bDiscardAll)
+						NumToRemove = 1;
+
+					if (NumToRemove && OffsetAndPakIndexOfSavedBlocked[CacheIndex].Num())
 					{
-						UE_LOG(LogPakFile, Verbose, TEXT("FPakReadRequest[%016llX, %016llX) Discard Cached"), Block.OffsetAndPakIndex, Block.OffsetAndPakIndex + Block.Size);
-						ClearBlock(Block);
-						return true;
+						NoneToRemove = false;
+						for (int32 Index = 0; Index < NumToRemove; Index++)
+						{
+							FJoinedOffsetAndPakIndex OffsetAndPakIndex = OffsetAndPakIndexOfSavedBlocked[CacheIndex][Index];
+							uint16 PakIndex = GetRequestPakIndex(OffsetAndPakIndex);
+							int64 Offset = GetRequestOffset(OffsetAndPakIndex);
+							FPakData& Pak = CachedPakData[PakIndex];
+							MaybeRemoveOverlappingNodesInIntervalTree<FCacheBlock>(
+								&Pak.CacheBlocks[(int32)EBlockStatus::Complete],
+								CacheBlockAllocator,
+								Offset,
+								Offset,
+								0,
+								Pak.MaxNode,
+								Pak.StartShift,
+								Pak.MaxShift,
+								[this](TIntervalTreeIndex BlockIndex) -> bool
+							{
+								FCacheBlock &Block = CacheBlockAllocator.Get(BlockIndex);
+								if (!Block.InRequestRefCount)
+								{
+									UE_LOG(LogPakFile, Verbose, TEXT("FPakReadRequest[%016llX, %016llX) Discard Cached"), Block.OffsetAndPakIndex, Block.OffsetAndPakIndex + Block.Size);
+									ClearBlock(Block);
+									return true;
+								}
+								return false;
+							}
+							);
+						}
+						OffsetAndPakIndexOfSavedBlocked[CacheIndex].RemoveAt(0, NumToRemove, false);
 					}
-					return false;
 				}
-				);
-
+				if (NoneToRemove)
+				{
+					break;
+				}
+			}
+		}
+		else
+		{
+			uint16 EndPakIndex = 65535;
+			if (StartPakIndex != 65535)
+			{
+				EndPakIndex = StartPakIndex + 1;
+			}
+			else
+			{
+				StartPakIndex = 0;
+				EndPakIndex = CachedPakData.Num();
+			}
 
+			// CachedFilesScopeLock is locked
+			for (uint16 RealPakIndex = StartPakIndex; RealPakIndex < EndPakIndex; RealPakIndex++)
+			{
+				int32 CacheIndex = CachedPakData[RealPakIndex].ActualPakFile->GetCacheIndex();
+				int32 NumToKeep = bDiscardAll ? 0 : GPakCache_NumUnreferencedBlocksToCache;
+				int32 NumToRemove = FMath::Max<int32>(0, OffsetAndPakIndexOfSavedBlocked[CacheIndex].Num() - NumToKeep);
+				if (NumToRemove)
+				{
+					for (int32 Index = 0; Index < NumToRemove; Index++)
+					{
+						FJoinedOffsetAndPakIndex OffsetAndPakIndex = OffsetAndPakIndexOfSavedBlocked[CacheIndex][Index];
+						uint16 PakIndex = GetRequestPakIndex(OffsetAndPakIndex);
+						int64 Offset = GetRequestOffset(OffsetAndPakIndex);
+						FPakData& Pak = CachedPakData[PakIndex];
+						MaybeRemoveOverlappingNodesInIntervalTree<FCacheBlock>(
+							&Pak.CacheBlocks[(int32)EBlockStatus::Complete],
+							CacheBlockAllocator,
+							Offset,
+							Offset,
+							0,
+							Pak.MaxNode,
+							Pak.StartShift,
+							Pak.MaxShift,
+							[this](TIntervalTreeIndex BlockIndex) -> bool
+						{
+							FCacheBlock &Block = CacheBlockAllocator.Get(BlockIndex);
+							if (!Block.InRequestRefCount)
+							{
+								UE_LOG(LogPakFile, Verbose, TEXT("FPakReadRequest[%016llX, %016llX) Discard Cached"), Block.OffsetAndPakIndex, Block.OffsetAndPakIndex + Block.Size);
+								ClearBlock(Block);
+								return true;
+							}
+							return false;
+						}
+						);
+					}
+					OffsetAndPakIndexOfSavedBlocked[CacheIndex].RemoveAt(0, NumToRemove, false);
+				}
 			}
-			OffsetAndPakIndexOfSavedBlocked.RemoveAt(0, NumToRemove, false);
 		}
 	}
 
@@ -1679,8 +1827,11 @@ private: // below here we assume CachedFilesScopeLock until we get to the next s
 				{
 					if (GPakCache_NumUnreferencedBlocksToCache && GetRequestOffset(Block.OffsetAndPakIndex) + Block.Size > OffsetOfLastByte) // last block
 					{
-						OffsetAndPakIndexOfSavedBlocked.Remove(Block.OffsetAndPakIndex);
-						OffsetAndPakIndexOfSavedBlocked.Add(Block.OffsetAndPakIndex);
+						uint16 BlocksPakIndex = GetRequestPakIndexLow(Block.OffsetAndPakIndex);
+						int32 BlocksCacheIndex = CachedPakData[BlocksPakIndex].ActualPakFile->GetCacheIndex();
+
+						OffsetAndPakIndexOfSavedBlocked[BlocksCacheIndex].Remove(Block.OffsetAndPakIndex);
+						OffsetAndPakIndexOfSavedBlocked[BlocksCacheIndex].Add(Block.OffsetAndPakIndex);
 						return false;
 					}
 					ClearBlock(Block);
@@ -1689,7 +1840,10 @@ private: // below here we assume CachedFilesScopeLock until we get to the next s
 				return false;
 			}
 			);
-			TrimCache();
+			if (!Pak.ActualPakFile->GetUnderlyingCacheTrimDisabled())
+			{
+				TrimCache(false, PakIndex);
+			}
 			OverlappingNodesInIntervalTree<FCacheBlock>(
 				Pak.CacheBlocks[(int32)EBlockStatus::InFlight],
 				CacheBlockAllocator,
@@ -2152,7 +2306,7 @@ private: // below here we assume CachedFilesScopeLock until we get to the next s
 		FAsyncFileCallBack CallbackFromLower =
 			[this, IndexToFill, bDoCheck](bool bWasCanceled, IAsyncReadRequest* Request)
 		{
-			if (bEnableSignatureChecks && bDoCheck)
+			if (SigningKey.IsValid() && bDoCheck)
 			{
 				StartSignatureCheck(bWasCanceled, Request, IndexToFill);
 			}
@@ -2206,7 +2360,7 @@ private: // below here we assume CachedFilesScopeLock until we get to the next s
 				BlockMemoryHighWater = BlockMemory;
 				SET_MEMORY_STAT(STAT_PakCacheHighWater, BlockMemoryHighWater);
 
-#if 0
+#if 1
 				static int64 LastPrint = 0;
 				if (BlockMemoryHighWater / 1024 / 1024 / 16 != LastPrint)
 				{
@@ -2263,6 +2417,8 @@ private: // below here we assume CachedFilesScopeLock until we get to the next s
 				NotifyComplete(Comp); // potentially scary recursion here
 			}
 		}
+
+		TrimCache();
 	}
 
 	bool StartNextRequest()
@@ -2377,12 +2533,12 @@ public:
 		NotifyRecursion--;
 	}
 
-	bool QueueRequest(IPakRequestor* Owner, FName File, int64 PakFileSize, int64 Offset, int64 Size, EAsyncIOPriorityAndFlags PriorityAndFlags)
+	bool QueueRequest(IPakRequestor* Owner, FPakFile* InActualPakFile, FName File, int64 PakFileSize, int64 Offset, int64 Size, EAsyncIOPriorityAndFlags PriorityAndFlags)
 	{
 		CSV_SCOPED_TIMING_STAT(FileIO, PakPrecacherQueueRequest);
 		check(Owner && File != NAME_None && Size > 0 && Offset >= 0 && Offset < PakFileSize && (PriorityAndFlags&AIOP_PRIORITY_MASK) >= AIOP_MIN && (PriorityAndFlags&AIOP_PRIORITY_MASK) <= AIOP_MAX);
 		FScopeLock Lock(&CachedFilesScopeLock);
-		uint16* PakIndexPtr = RegisterPakFile(File, PakFileSize);
+		uint16* PakIndexPtr = RegisterPakFile(InActualPakFile, File, PakFileSize);
 		if (PakIndexPtr == nullptr)
 		{
 			return false;
@@ -2416,7 +2572,7 @@ public:
 		{
 			UE_LOG(LogPakFile, Verbose, TEXT("FPakReadRequest[%016llX, %016llX) QueueRequest COLD"), RequestOffsetAndPakIndex, RequestOffsetAndPakIndex + Request.Size);
 		}
-
+		TrimCache();
 		return true;
 	}
 
@@ -2483,135 +2639,138 @@ public:
 	void Unmount(FName PakFile)
 	{
 		FScopeLock Lock(&CachedFilesScopeLock);
-		uint16* PakIndexPtr = CachedPaks.Find(PakFile);
-		if (!PakIndexPtr)
-		{
-			UE_LOG(LogPakFile, Log, TEXT("Pak file %s was never used, so nothing to unmount"), *PakFile.ToString());
-			return; // never used for anything, nothing to check or clean up
-		}
-		TrimCache(true);
-		uint16 PakIndex = *PakIndexPtr;
-		FPakData& Pak = CachedPakData[PakIndex];
-		int64 Offset = MakeJoinedRequest(PakIndex, 0);
-
-		bool bHasOutstandingRequests = false;
 
-		OverlappingNodesInIntervalTree<FCacheBlock>(
-			Pak.CacheBlocks[(int32)EBlockStatus::Complete],
-			CacheBlockAllocator,
-			0,
-			Offset + Pak.TotalSize - 1,
-			0,
-			Pak.MaxNode,
-			Pak.StartShift,
-			Pak.MaxShift,
-			[&bHasOutstandingRequests](TIntervalTreeIndex BlockIndex) -> bool
-		{
-			check(!"Pak cannot be unmounted with outstanding requests");
-			bHasOutstandingRequests = true;
-			return false;
-		}
-		);
-		OverlappingNodesInIntervalTree<FCacheBlock>(
-			Pak.CacheBlocks[(int32)EBlockStatus::InFlight],
-			CacheBlockAllocator,
-			0,
-			Offset + Pak.TotalSize - 1,
-			0,
-			Pak.MaxNode,
-			Pak.StartShift,
-			Pak.MaxShift,
-			[&bHasOutstandingRequests](TIntervalTreeIndex BlockIndex) -> bool
-		{
-			check(!"Pak cannot be unmounted with outstanding requests");
-			bHasOutstandingRequests = true;
-			return false;
-		}
-		);
-		for (int32 Priority = AIOP_MAX;; Priority--)
-		{
-			OverlappingNodesInIntervalTree<FPakInRequest>(
-				Pak.InRequests[Priority][(int32)EInRequestStatus::InFlight],
-				InRequestAllocator,
-				0,
-				Offset + Pak.TotalSize - 1,
-				0,
-				Pak.MaxNode,
-				Pak.StartShift,
-				Pak.MaxShift,
-				[&bHasOutstandingRequests](TIntervalTreeIndex BlockIndex) -> bool
-			{
-				check(!"Pak cannot be unmounted with outstanding requests");
-				bHasOutstandingRequests = true;
-				return false;
-			}
-			);
-			OverlappingNodesInIntervalTree<FPakInRequest>(
-				Pak.InRequests[Priority][(int32)EInRequestStatus::Complete],
-				InRequestAllocator,
-				0,
-				Offset + Pak.TotalSize - 1,
-				0,
-				Pak.MaxNode,
-				Pak.StartShift,
-				Pak.MaxShift,
-				[&bHasOutstandingRequests](TIntervalTreeIndex BlockIndex) -> bool
-			{
-				check(!"Pak cannot be unmounted with outstanding requests");
-				bHasOutstandingRequests = true;
-				return false;
-			}
-			);
-			OverlappingNodesInIntervalTree<FPakInRequest>(
-				Pak.InRequests[Priority][(int32)EInRequestStatus::Waiting],
-				InRequestAllocator,
-				0,
-				Offset + Pak.TotalSize - 1,
-				0,
-				Pak.MaxNode,
-				Pak.StartShift,
-				Pak.MaxShift,
-				[&bHasOutstandingRequests](TIntervalTreeIndex BlockIndex) -> bool
-			{
-				check(!"Pak cannot be unmounted with outstanding requests");
-				bHasOutstandingRequests = true;
-				return false;
-			}
-			);
-			if (Priority == AIOP_MIN)
-			{
-				break;
-			}
-		}
-		if (!bHasOutstandingRequests)
+		for (TMap<FPakFile*, uint16>::TIterator It(CachedPaks); It; ++It)
 		{
-			UE_LOG(LogPakFile, Log, TEXT("Pak file %s removed from pak precacher."), *PakFile.ToString());
-			CachedPaks.Remove(PakFile);
-			check(Pak.Handle);
-			delete Pak.Handle;
-			Pak.Handle = nullptr;
-			int32 NumToTrim = 0;
-			for (int32 Index = CachedPakData.Num() - 1; Index >= 0; Index--)
+			if( It->Key->GetFilenameName() == PakFile )
 			{
-				if (!CachedPakData[Index].Handle)
+				uint16 PakIndex = It->Value;
+				TrimCache(true);
+				FPakData& Pak = CachedPakData[PakIndex];
+				int64 Offset = MakeJoinedRequest(PakIndex, 0);
+
+				bool bHasOutstandingRequests = false;
+
+				OverlappingNodesInIntervalTree<FCacheBlock>(
+					Pak.CacheBlocks[(int32)EBlockStatus::Complete],
+					CacheBlockAllocator,
+					0,
+					Offset + Pak.TotalSize - 1,
+					0,
+					Pak.MaxNode,
+					Pak.StartShift,
+					Pak.MaxShift,
+					[&bHasOutstandingRequests](TIntervalTreeIndex BlockIndex) -> bool
+				{
+					check(!"Pak cannot be unmounted with outstanding requests");
+					bHasOutstandingRequests = true;
+					return false;
+				}
+				);
+				OverlappingNodesInIntervalTree<FCacheBlock>(
+					Pak.CacheBlocks[(int32)EBlockStatus::InFlight],
+					CacheBlockAllocator,
+					0,
+					Offset + Pak.TotalSize - 1,
+					0,
+					Pak.MaxNode,
+					Pak.StartShift,
+					Pak.MaxShift,
+					[&bHasOutstandingRequests](TIntervalTreeIndex BlockIndex) -> bool
 				{
-					NumToTrim++;
+					check(!"Pak cannot be unmounted with outstanding requests");
+					bHasOutstandingRequests = true;
+					return false;
+				}
+				);
+				for (int32 Priority = AIOP_MAX;; Priority--)
+				{
+					OverlappingNodesInIntervalTree<FPakInRequest>(
+						Pak.InRequests[Priority][(int32)EInRequestStatus::InFlight],
+						InRequestAllocator,
+						0,
+						Offset + Pak.TotalSize - 1,
+						0,
+						Pak.MaxNode,
+						Pak.StartShift,
+						Pak.MaxShift,
+						[&bHasOutstandingRequests](TIntervalTreeIndex BlockIndex) -> bool
+					{
+						check(!"Pak cannot be unmounted with outstanding requests");
+						bHasOutstandingRequests = true;
+						return false;
+					}
+					);
+					OverlappingNodesInIntervalTree<FPakInRequest>(
+						Pak.InRequests[Priority][(int32)EInRequestStatus::Complete],
+						InRequestAllocator,
+						0,
+						Offset + Pak.TotalSize - 1,
+						0,
+						Pak.MaxNode,
+						Pak.StartShift,
+						Pak.MaxShift,
+						[&bHasOutstandingRequests](TIntervalTreeIndex BlockIndex) -> bool
+					{
+						check(!"Pak cannot be unmounted with outstanding requests");
+						bHasOutstandingRequests = true;
+						return false;
+					}
+					);
+					OverlappingNodesInIntervalTree<FPakInRequest>(
+						Pak.InRequests[Priority][(int32)EInRequestStatus::Waiting],
+						InRequestAllocator,
+						0,
+						Offset + Pak.TotalSize - 1,
+						0,
+						Pak.MaxNode,
+						Pak.StartShift,
+						Pak.MaxShift,
+						[&bHasOutstandingRequests](TIntervalTreeIndex BlockIndex) -> bool
+					{
+						check(!"Pak cannot be unmounted with outstanding requests");
+						bHasOutstandingRequests = true;
+						return false;
+					}
+					);
+					if (Priority == AIOP_MIN)
+					{
+						break;
+					}
+				}
+				if (!bHasOutstandingRequests)
+				{
+					UE_LOG(LogPakFile, Log, TEXT("Pak file %s removed from pak precacher."), *PakFile.ToString());
+					It.RemoveCurrent();
+					check(Pak.Handle);
+					delete Pak.Handle;
+					Pak.Handle = nullptr;
+					int32 NumToTrim = 0;
+					for (int32 Index = CachedPakData.Num() - 1; Index >= 0; Index--)
+					{
+						if (!CachedPakData[Index].Handle)
+						{
+							NumToTrim++;
+						}
+						else
+						{
+							break;
+						}
+					}
+					if (NumToTrim)
+					{
+						CachedPakData.RemoveAt(CachedPakData.Num() - NumToTrim, NumToTrim);
+						LastReadRequest = 0;
+					}
 				}
 				else
 				{
-					break;
+					UE_LOG(LogPakFile, Log, TEXT("Pak file %s was NOT removed from pak precacher because it had outstanding requests."), *PakFile.ToString());
 				}
+
 			}
-			if (NumToTrim)
-			{
-				CachedPakData.RemoveAt(CachedPakData.Num() - NumToTrim, NumToTrim);
-				LastReadRequest = 0;
-			}
-		}
-		else
-		{
-			UE_LOG(LogPakFile, Log, TEXT("Pak file %s was NOT removed from pak precacher because it had outstanding requests."), *PakFile.ToString());
 		}
+
 	}
 
 
@@ -2820,13 +2979,13 @@ class FPakReadRequest : public FPakReadRequestBase
 {
 public:
 
-	FPakReadRequest(FName InPakFile, int64 PakFileSize, FAsyncFileCallBack* CompleteCallback, int64 InOffset, int64 InBytesToRead, EAsyncIOPriorityAndFlags InPriorityAndFlags, uint8* UserSuppliedMemory, bool bInInternalRequest = false, FCachedAsyncBlock* InBlockPtr = nullptr)
+	FPakReadRequest(FPakFile* InActualPakFile, FName InPakFile, int64 PakFileSize, FAsyncFileCallBack* CompleteCallback, int64 InOffset, int64 InBytesToRead, EAsyncIOPriorityAndFlags InPriorityAndFlags, uint8* UserSuppliedMemory, bool bInInternalRequest = false, FCachedAsyncBlock* InBlockPtr = nullptr)
 		: FPakReadRequestBase(InPakFile, PakFileSize, CompleteCallback, InOffset, InBytesToRead, InPriorityAndFlags, UserSuppliedMemory, bInInternalRequest, InBlockPtr)
 	{
 		check(Offset >= 0 && BytesToRead > 0);
 		check(bInternalRequest || ( InPriorityAndFlags & AIOP_FLAG_PRECACHE ) == 0 || !bUserSuppliedMemory); // you never get bits back from a precache request, so why supply memory?
 
-		if (!FPakPrecacher::Get().QueueRequest(this, InPakFile, PakFileSize, Offset, BytesToRead, InPriorityAndFlags))
+		if (!FPakPrecacher::Get().QueueRequest(this, InActualPakFile, InPakFile, PakFileSize, Offset, BytesToRead, InPriorityAndFlags))
 		{
 			bRequestOutstanding = false;
 			SetComplete();
@@ -2891,7 +3050,7 @@ class FPakEncryptedReadRequest : public FPakReadRequestBase
 
 public:
 
-	FPakEncryptedReadRequest(FName InPakFile, int64 PakFileSize, FAsyncFileCallBack* CompleteCallback, int64 InPakFileStartOffset, int64 InFileOffset, int64 InBytesToRead, EAsyncIOPriorityAndFlags InPriorityAndFlags, uint8* UserSuppliedMemory, const FGuid& InEncryptionKeyGuid, bool bInInternalRequest = false, FCachedAsyncBlock* InBlockPtr = nullptr)
+	FPakEncryptedReadRequest(FPakFile* InActualPakFile, FName InPakFile, int64 PakFileSize, FAsyncFileCallBack* CompleteCallback, int64 InPakFileStartOffset, int64 InFileOffset, int64 InBytesToRead, EAsyncIOPriorityAndFlags InPriorityAndFlags, uint8* UserSuppliedMemory, const FGuid& InEncryptionKeyGuid, bool bInInternalRequest = false, FCachedAsyncBlock* InBlockPtr = nullptr)
 		: FPakReadRequestBase(InPakFile, PakFileSize, CompleteCallback, InPakFileStartOffset + InFileOffset, InBytesToRead, InPriorityAndFlags, UserSuppliedMemory, bInInternalRequest, InBlockPtr)
 		, OriginalOffset(InPakFileStartOffset + InFileOffset)
 		, OriginalSize(InBytesToRead)
@@ -2900,7 +3059,7 @@ public:
 		Offset = InPakFileStartOffset + AlignDown(InFileOffset, FAES::AESBlockSize);
 		BytesToRead = Align(InFileOffset + InBytesToRead, FAES::AESBlockSize) - AlignDown(InFileOffset, FAES::AESBlockSize);
 
-		if (!FPakPrecacher::Get().QueueRequest(this, InPakFile, PakFileSize, Offset, BytesToRead, InPriorityAndFlags))
+		if (!FPakPrecacher::Get().QueueRequest(this, InActualPakFile, InPakFile, PakFileSize, Offset, BytesToRead, InPriorityAndFlags))
 		{
 			bRequestOutstanding = false;
 			SetComplete();
@@ -3303,6 +3462,7 @@ void FPakPrecacher::DoSignatureCheck(bool bWasCanceled, IAsyncReadRequest* Reque
 class FPakAsyncReadFileHandle final : public IAsyncReadFileHandle
 {
 	FName PakFile;
+	FPakFile* ActualPakFile;
 	int64 PakFileSize;
 	int64 OffsetInPak;
 	int64 UncompressedFileSize;
@@ -3332,6 +3492,7 @@ class FPakAsyncReadFileHandle final : public IAsyncReadFileHandle
 public:
 	FPakAsyncReadFileHandle(const FPakEntry* InFileEntry, FPakFile* InPakFile, const TCHAR* Filename)
 		: PakFile(InPakFile->GetFilenameName())
+		, ActualPakFile(InPakFile)
 		, PakFileSize(InPakFile->TotalSize())
 		, FileEntry(*InFileEntry)
 		, NumLiveRawRequests(0)
@@ -3401,11 +3562,11 @@ public:
 
 			if (FileEntry.IsEncrypted())
 			{
-				return new FPakEncryptedReadRequest(PakFile, PakFileSize, CompleteCallback, OffsetInPak, Offset, BytesToRead, PriorityAndFlags, UserSuppliedMemory, EncryptionKeyGuid);
+				return new FPakEncryptedReadRequest(ActualPakFile, PakFile, PakFileSize, CompleteCallback, OffsetInPak, Offset, BytesToRead, PriorityAndFlags, UserSuppliedMemory, EncryptionKeyGuid);
 			}
 			else
 			{
-				return new FPakReadRequest(PakFile, PakFileSize, CompleteCallback, OffsetInPak + Offset, BytesToRead, PriorityAndFlags, UserSuppliedMemory);
+				return new FPakReadRequest(ActualPakFile, PakFile, PakFileSize, CompleteCallback, OffsetInPak + Offset, BytesToRead, PriorityAndFlags, UserSuppliedMemory);
 			}
 		}
 		bool bAnyUnfinished = false;
@@ -3457,7 +3618,7 @@ public:
 			Block.RawSize = Align(Block.RawSize, FAES::AESBlockSize);
 		}
 		NumLiveRawRequests++;
-		Block.RawRequest = new FPakReadRequest(PakFile, PakFileSize, &ReadCallbackFunction, FileEntry.CompressionBlocks[BlockIndex].CompressedStart + CompressedChunkOffset, Block.RawSize, PriorityAndFlags, nullptr, true, &Block);
+		Block.RawRequest = new FPakReadRequest(ActualPakFile, PakFile, PakFileSize, &ReadCallbackFunction, FileEntry.CompressionBlocks[BlockIndex].CompressedStart + CompressedChunkOffset, Block.RawSize, PriorityAndFlags, nullptr, true, &Block);
 	}
 	void RawReadCallback(bool bWasCancelled, IAsyncReadRequest* InRequest)
 	{
@@ -4358,6 +4519,9 @@ FPakFile::FPakFile(const TCHAR* Filename, bool bIsSigned)
 	, bAttemptedPakEntryShrink(false)
 	, bAttemptedPakFilenameUnload(false)
  	, MappedFileHandle(nullptr)
+	, CacheType(ECacheType::Shared)
+	, CacheIndex(-1)
+	, UnderlyingCacheTrimDisabled(false)
 {
 	FArchive* Reader = GetSharedReader(NULL);
 	if (Reader)
@@ -4368,7 +4532,7 @@ FPakFile::FPakFile(const TCHAR* Filename, bool bIsSigned)
 }
 #endif
 
-FPakFile::FPakFile(IPlatformFile* LowerLevel, const TCHAR* Filename, bool bIsSigned)
+FPakFile::FPakFile(IPlatformFile* LowerLevel, const TCHAR* Filename, bool bIsSigned, bool bLoadIndex)
 	: PakFilename(Filename)
 	, PakFilenameName(Filename)
 	, FilenameHashesIndex(nullptr)
@@ -4385,12 +4549,15 @@ FPakFile::FPakFile(IPlatformFile* LowerLevel, const TCHAR* Filename, bool bIsSig
 	, bAttemptedPakEntryShrink(false)
 	, bAttemptedPakFilenameUnload(false)
 	, MappedFileHandle(nullptr)
+	, CacheType(ECacheType::Shared)
+	, CacheIndex(-1)
+	, UnderlyingCacheTrimDisabled(false)
 {
 	FArchive* Reader = GetSharedReader(LowerLevel);
 	if (Reader)
 	{
 		Timestamp = LowerLevel->GetTimeStamp(Filename);
-		Initialize(Reader);
+		Initialize(Reader, bLoadIndex);
 	}
 }
 
@@ -4407,6 +4574,9 @@ FPakFile::FPakFile(FArchive* Archive)
 	, bFilenamesRemoved(false)
 	, ChunkID(INDEX_NONE)
 	, MappedFileHandle(nullptr)
+	, CacheType(ECacheType::Shared)
+	, CacheIndex(-1)
+	, UnderlyingCacheTrimDisabled(false)
 {
 	Initialize(Archive);
 }
@@ -4468,7 +4638,7 @@ FArchive* FPakFile::SetupSignedPakReader(FArchive* ReaderArchive, const TCHAR* F
 	return ReaderArchive;
 }
 
-void FPakFile::Initialize(FArchive* Reader)
+void FPakFile::Initialize(FArchive* Reader, bool bLoadIndex)
 {
 	CachedTotalSize = Reader->TotalSize();
 	bool bShouldLoad = false;
@@ -4510,7 +4680,10 @@ void FPakFile::Initialize(FArchive* Reader)
 		// If we aren't using a dynamic encryption key, process the pak file using the embedded key
 		if (!Info.EncryptionKeyGuid.IsValid() || GetRegisteredEncryptionKeys().HasKey(Info.EncryptionKeyGuid))
 		{
-			LoadIndex(Reader);
+			if (bLoadIndex)
+			{
+				LoadIndex(Reader);
+			}
 
 			if (FParse::Param(FCommandLine::Get(), TEXT("checkpak")))
 			{
@@ -4812,12 +4985,17 @@ bool FPakFile::UnloadPakEntryFilenames(TMap<uint64, FPakEntry>& CrossPakCollisio
 
 		// Build the list of hashes from the Index based on the starting hash.
 		int32 EntryIndex = 0;
+		FString FinalFilename;
+		FinalFilename.Reset(1024);
 		for (TMap<FString, FPakDirectory>::TConstIterator It(Index); (It && !bHasCollision); ++It)
 		{
 			for (FPakDirectory::TConstIterator DirectoryIt(It.Value()); DirectoryIt; ++DirectoryIt)
 			{
-				FString FinalFilename = It.Key() / DirectoryIt.Key();
-				uint64 FilenameHash = FFnv::MemFnv64(*FinalFilename.ToLower(), FinalFilename.Len() * sizeof(TCHAR), FilenameStartHash);
+				FinalFilename.Reset();
+				FinalFilename.Append(It.Key());
+				FinalFilename.PathAppend(*DirectoryIt.Key(), DirectoryIt.Key().Len());
+				FinalFilename.ToLowerInline();
+				uint64 FilenameHash = FFnv::MemFnv64(*FinalFilename, FinalFilename.Len() * sizeof(TCHAR), FilenameStartHash);
 				MiniFileEntries[EntryIndex].FilenameHash = FilenameHash;
 				MiniFileEntries[EntryIndex].EntryIndex = DirectoryIt.Value();
 				++EntryIndex;
@@ -5755,9 +5933,9 @@ bool FPakPlatformFile::Initialize(IPlatformFile* Inner, const TCHAR* CmdLine)
 	GameUserSettingsIniFilename = TEXT("GameUserSettings.ini");
 #endif
 
-	// Signed if we have keys, and are not running with fileopenlog (currently results in a deadlock).
-	bSigned = FCoreDelegates::GetPakSigningKeysDelegate().IsBound() && !FParse::Param(FCommandLine::Get(), TEXT("fileopenlog"));
-
+	// signed if we have keys, and are not running with fileopenlog (currently results in a deadlock).
+	bSigned = GetPakSigningKey().IsValid() && !FParse::Param(FCommandLine::Get(), TEXT("fileopenlog"));;
+		
 	// Find and mount pak files from the specified directories.
 	TArray<FString> PakFolders;
 	GetPakFolders(FCommandLine::Get(), PakFolders);
@@ -5783,7 +5961,7 @@ void FPakPlatformFile::InitializeNewAsyncIO()
 #if !WITH_EDITOR
 	if (FPlatformProcess::SupportsMultithreading() && !FParse::Param(FCommandLine::Get(), TEXT("FileOpenLog")))
 	{
-		FPakPrecacher::Init(LowerLevel, FCoreDelegates::GetPakSigningKeysDelegate().IsBound());
+		FPakPrecacher::Init(LowerLevel, GetPakSigningKey());
 	}
 	else
 #endif
@@ -5797,6 +5975,7 @@ void FPakPlatformFile::InitializeNewAsyncIO()
 void FPakPlatformFile::OptimizeMemoryUsageForMountedPaks()
 {
 #if !(IS_PROGRAM || WITH_EDITOR)
+	FSlowHeartBeatScope SuspendHeartBeat;
 	bool bUnloadPakEntryFilenamesIfPossible = FParse::Param(FCommandLine::Get(), TEXT("unloadpakentryfilenames"));
 	GConfig->GetBool(TEXT("Pak"), TEXT("UnloadPakEntryFilenamesIfPossible"), bUnloadPakEntryFilenamesIfPossible, GEngineIni);
 
@@ -5827,13 +6006,14 @@ void FPakPlatformFile::OptimizeMemoryUsageForMountedPaks()
 #endif
 }
 
-bool FPakPlatformFile::Mount(const TCHAR* InPakFilename, uint32 PakOrder, const TCHAR* InPath /*= NULL*/)
+
+bool FPakPlatformFile::Mount(const TCHAR* InPakFilename, uint32 PakOrder, const TCHAR* InPath /*= NULL*/, bool bLoadIndex /*= true*/)
 {
 	bool bSuccess = false;
 	TSharedPtr<IFileHandle> PakHandle = MakeShareable(LowerLevel->OpenRead(InPakFilename));
 	if (PakHandle.IsValid())
 	{
-		FPakFile* Pak = new FPakFile(LowerLevel, InPakFilename, bSigned);
+		FPakFile* Pak = new FPakFile(LowerLevel, InPakFilename, bSigned, bLoadIndex);
 		if (Pak->IsValid())
 		{
 			if (GetRegisteredEncryptionKeys().HasKey(Pak->GetInfo().EncryptionKeyGuid))
@@ -6180,7 +6360,7 @@ void FPakPlatformFile::RegisterEncryptionKey(const FGuid& InGuid, const FAES::FA
 			OptimizeMemoryUsageForMountedPaks();
 		}
 
-		UE_CLOG(InGuid.IsValid(), LogPakFile, Log, TEXT("Registered encryption key '%s': %d pak files mounted, %d remain pending"), *InGuid.ToString(), NumMounted, PendingEncryptedPakFiles.Num());
+		UE_LOG(LogPakFile, Log, TEXT("Registered encryption key '%s': %d pak files mounted, %d remain pending"), *InGuid.ToString(), NumMounted, PendingEncryptedPakFiles.Num());
 	}
 }
 
@@ -6414,4 +6594,101 @@ public:
 	TUniquePtr<IPlatformFile> Singleton;
 };
 
+void FPakFile::AddSpecialFile(FPakEntry Entry, const FString& Filename)
+{
+	MakeDirectoryFromPath(MountPoint);
+
+	int32 EntryIndex = Files.Num();
+
+	// Add new file info.
+	Files.Add(Entry);
+	NumEntries++;
+
+	// Construct Index of all directories in pak file.
+
+
+	FString CleanFilename(Filename.Mid(MountPoint.Len() ));
+
+	FString Path = FPaths::GetPath(CleanFilename);
+	FString Path2 = FPaths::GetPath(Filename);
+	MakeDirectoryFromPath(Path);
+
+
+
+	FPakDirectory* Directory = Index.Find(Path); 
+	if (Directory != NULL)
+	{
+		Directory->Add(FPaths::GetCleanFilename(CleanFilename), EntryIndex);
+	}
+	else
+	{
+		FPakDirectory& NewDirectory = Index.Add(Path);
+		NewDirectory.Add(FPaths::GetCleanFilename(CleanFilename), EntryIndex);
+
+		// add the parent directories up to the mount point
+		while (MountPoint != Path)
+		{
+			Path = Path.Left(Path.Len() - 1);
+			int32 Offset = 0;
+			if (Path.FindLastChar('/', Offset))
+			{
+				Path = Path.Left(Offset);
+				MakeDirectoryFromPath(Path);
+				if (Index.Find(Path) == NULL)
+				{
+					Index.Add(Path);
+				}
+			}
+			else
+			{
+				Path = MountPoint;
+			}
+		}
+	}
+}
+
+
+
+void FPakPlatformFile::MakeUniquePakFilesForTheseFiles(TArray<TArray<FString>> InFiles)
+{
+	for (int k = 0; k < InFiles.Num(); k++)
+	{
+		FPakFile* NewPakFile = NULL;
+		for (int i = 0; i < InFiles[k].Num(); i++)
+		{
+			FPakEntry FileEntry;
+			FPakFile* PakFile = NULL;
+			bool bFoundEntry = FindFileInPakFiles(*InFiles[k][i], &PakFile, &FileEntry);
+			if (bFoundEntry && PakFile && PakFile->GetFilenameName() != NAME_None)
+			{
+				if (NewPakFile == nullptr && Mount(*PakFile->GetFilename(), 500, *PakFile->MountPoint, false))
+				{
+					// we successfully mounted the file, find the empty pak file we just added.
+					for (int j = 0; j < PakFiles.Num(); j++)
+					{
+						if (PakFiles[j].PakFile->Files.Num() == 0 && PakFiles[j].PakFile->CachedTotalSize == PakFile->CachedTotalSize)
+						{
+							NewPakFile = PakFiles[j].PakFile;
+							break;
+						}
+					}
+					if (NewPakFile != nullptr)
+					{
+						NewPakFile->SetCacheType(FPakFile::ECacheType::Individual);
+					}
+				}
+
+				if (NewPakFile != nullptr)
+				{
+//					NewPakFile->SetUnderlyingCacheTrimDisabled(true);
+					NewPakFile->AddSpecialFile(FileEntry, *InFiles[k][i]);
+				}
+
+			}
+		}
+
+	}
+}
+
+
 IMPLEMENT_MODULE(FPakFileModule, PakFile);
diff --git a/Engine/Source/Runtime/PakFile/Public/IPlatformFilePak.h b/Engine/Source/Runtime/PakFile/Public/IPlatformFilePak.h
--- a/Engine/Source/Runtime/PakFile/Public/IPlatformFilePak.h
+++ b/Engine/Source/Runtime/PakFile/Public/IPlatformFilePak.h
@@ -577,8 +577,31 @@ class PAKFILE_API FPakFile : FNoncopyable
 	FArchive* CreatePakReader(IFileHandle& InHandle, const TCHAR* Filename);
 	FArchive* SetupSignedPakReader(FArchive* Reader, const TCHAR* Filename);
 
+
 public:
+	/** Pak files can share a cache or have their own */
+	enum class ECacheType : uint8
+	{
+		Shared,
+		Individual,
+	};
+private:
+	/** The type of cache this pak file should have */
+	ECacheType	CacheType;
+	/** The index of this pak file into the cache array, -1 = not initialized */
+	int32		CacheIndex;
+	/** Allow the cache of a pak file to never shrink, should be used with caution, it will burn memory */
+	bool UnderlyingCacheTrimDisabled;
+
+public:
+
+	void SetUnderlyingCacheTrimDisabled(bool InUnderlyingCacheTrimDisabled) { UnderlyingCacheTrimDisabled = InUnderlyingCacheTrimDisabled; }
+	bool GetUnderlyingCacheTrimDisabled(void) { return UnderlyingCacheTrimDisabled; }
 
+	void SetCacheType(ECacheType InCacheType) { CacheType = InCacheType; }
+	ECacheType GetCacheType(void) { return CacheType; }
+	void SetCacheIndex(int32 InCacheIndex) { CacheIndex = InCacheIndex; }
+	int32 GetCacheIndex(void) { return CacheIndex; }
 #if IS_PROGRAM
 	/**
 	* Opens a pak file given its filename.
@@ -596,7 +619,7 @@ public:
 	 * @param Filename Filename.
 	 * @param bIsSigned = true if the pak is signed.
 	 */
-	FPakFile(IPlatformFile* LowerLevel, const TCHAR* Filename, bool bIsSigned);
+	FPakFile(IPlatformFile* LowerLevel, const TCHAR* Filename, bool bIsSigned, bool bLoadIndex = true);
 
 	/**
 	 * Creates a pak file using the supplied archive.
@@ -979,13 +1002,18 @@ private:
 	/**
 	 * Initializes the pak file.
 	 */
-	void Initialize(FArchive* Reader);
+	void Initialize(FArchive* Reader, bool bLoadIndex = true);
 
 	/**
 	 * Loads and initializes pak file index.
 	 */
 	void LoadIndex(FArchive* Reader);
 
+	/**
+	 * Manually add a file to a pak file,
+	 */
+	void AddSpecialFile(FPakEntry Entry, const FString& Filename);
+
 	/**
 	 * Decodes a bit-encoded pak entry.
 	 *
@@ -1621,13 +1649,18 @@ public:
 	 * @param InPakFilename Pak filename.
 	 * @param InPath Path to mount the pak at.
 	 */
-	bool Mount(const TCHAR* InPakFilename, uint32 PakOrder, const TCHAR* InPath = NULL);
+	bool Mount(const TCHAR* InPakFilename, uint32 PakOrder, const TCHAR* InPath = NULL, bool bLoadIndex = true);
 
 	bool Unmount(const TCHAR* InPakFilename);
 
 	int32 MountAllPakFiles(const TArray<FString>& PakFolders);
 	int32 MountAllPakFiles(const TArray<FString>& PakFolders, const FString& WildCard);
 
+	/**
+	 * Make unique in memory pak files from a list of named files
+	 */
+	virtual void MakeUniquePakFilesForTheseFiles(TArray<TArray<FString>> InFiles);
+
 
 	/**
 	 * Finds a file in the specified pak files.
