Fix build break

diff --git a/Engine/Source/Runtime/PakFile/Private/IPlatformFilePak.cpp b/Engine/Source/Runtime/PakFile/Private/IPlatformFilePak.cpp
--- a/Engine/Source/Runtime/PakFile/Private/IPlatformFilePak.cpp
+++ b/Engine/Source/Runtime/PakFile/Private/IPlatformFilePak.cpp
@@ -4519,7 +4519,7 @@ FPakFile::FPakFile(const TCHAR* Filename, bool bIsSigned)
 	, bAttemptedPakEntryShrink(false)
 	, bAttemptedPakFilenameUnload(false)
  	, MappedFileHandle(nullptr)
-	, CacheType(ECacheType::Shared)
+	, CacheType(FPakFile::ECacheType::Shared)
 	, CacheIndex(-1)
 	, UnderlyingCacheTrimDisabled(false)
 {
@@ -4549,7 +4549,7 @@ FPakFile::FPakFile(IPlatformFile* LowerLevel, const TCHAR* Filename, bool bIsSig
 	, bAttemptedPakEntryShrink(false)
 	, bAttemptedPakFilenameUnload(false)
 	, MappedFileHandle(nullptr)
-	, CacheType(ECacheType::Shared)
+	, CacheType(FPakFile::ECacheType::Shared)
 	, CacheIndex(-1)
 	, UnderlyingCacheTrimDisabled(false)
 {
@@ -4574,7 +4574,7 @@ FPakFile::FPakFile(FArchive* Archive)
 	, bFilenamesRemoved(false)
 	, ChunkID(INDEX_NONE)
 	, MappedFileHandle(nullptr)
-	, CacheType(ECacheType::Shared)
+	, CacheType(FPakFile::ECacheType::Shared)
 	, CacheIndex(-1)
 	, UnderlyingCacheTrimDisabled(false)
 {
