Fixing compile error apparently due to bad robomerge in CL 7689301.

diff --git a/Engine/Source/Runtime/PakFile/Private/IPlatformFilePak.cpp b/Engine/Source/Runtime/PakFile/Private/IPlatformFilePak.cpp
--- a/Engine/Source/Runtime/PakFile/Private/IPlatformFilePak.cpp
+++ b/Engine/Source/Runtime/PakFile/Private/IPlatformFilePak.cpp
@@ -239,13 +239,21 @@ void FPakPlatformFile::GetPakEncryptionKey(FAES::FAESKey& OutKey, const FGuid& I
 	}
 }
 
-TSharedPtr<FRSA::FKey, ESPMode::ThreadSafe> FPakPlatformFile::GetPakSigningKey()
+TMap<FName, TSharedPtr<const struct FPakSignatureFile, ESPMode::ThreadSafe>> FPakPlatformFile::PakSignatureFileCache;
+FCriticalSection FPakPlatformFile::PakSignatureFileCacheLock;
+TSharedPtr<const struct FPakSignatureFile, ESPMode::ThreadSafe> FPakPlatformFile::GetPakSignatureFile(const TCHAR* InFilename)
 {
-	static TSharedPtr<FRSA::FKey, ESPMode::ThreadSafe> Key;
-	static FCriticalSection Lock;
-	Lock.Lock();
+	FScopeLock Lock(&PakSignatureFileCacheLock);
 
-	if (!Key.IsValid())
+	FName FilenameFName(InFilename);
+	if (const TSharedPtr<const struct FPakSignatureFile, ESPMode::ThreadSafe>* SignaturesFile = PakSignatureFileCache.Find(FilenameFName))
+	{
+		return *SignaturesFile;
+	}
+
+	static FRSAKeyHandle PublicKey = InvalidRSAKeyHandle;
+	static bool bInitializedPublicKey = false;
+	if (!bInitializedPublicKey)
 	{
 		FCoreDelegates::FPakSigningKeysDelegate& Delegate = FCoreDelegates::GetPakSigningKeysDelegate();
 		if (Delegate.IsBound())
@@ -253,44 +261,37 @@ TSharedPtr<FRSA::FKey, ESPMode::ThreadSafe> FPakPlatformFile::GetPakSigningKey()
 			TArray<uint8> Exponent;
 			TArray<uint8> Modulus;
 			Delegate.Execute(Exponent, Modulus);
-			Key = FRSA::CreateKey(Exponent, TArray<uint8>(), Modulus);
+			PublicKey = FRSA::CreateKey(Exponent, TArray<uint8>(), Modulus);
 		}
+		bInitializedPublicKey = true;
 	}
-	
-	Lock.Unlock();
-	return Key;
-}
 
-TMap<FName, TSharedPtr<const struct FPakSignatureFile, ESPMode::ThreadSafe>> FPakPlatformFile::PakSignatureFileCache;
-FCriticalSection FPakPlatformFile::PakSignatureFileCacheLock;
-TSharedPtr<const struct FPakSignatureFile, ESPMode::ThreadSafe> FPakPlatformFile::GetPakSignatureFile(const TCHAR* InFilename)
-{
-	FScopeLock Lock(&PakSignatureFileCacheLock);
-	FName FilenameFName(InFilename);
-	if (const TSharedPtr<const struct FPakSignatureFile, ESPMode::ThreadSafe>* SignaturesFile = PakSignatureFileCache.Find(FilenameFName))
-	{
-		return *SignaturesFile;
-	}
+	TSharedPtr<FPakSignatureFile, ESPMode::ThreadSafe> SignaturesFile;
 
-	TSharedPtr<FPakSignatureFile, ESPMode::ThreadSafe> SignaturesFile = MakeShared<FPakSignatureFile, ESPMode::ThreadSafe>();
-	FString SignaturesFilename = FPaths::ChangeExtension(InFilename, TEXT("sig"));
-	FArchive* Reader = IFileManager::Get().CreateFileReader(*SignaturesFilename);
-	if (Reader != nullptr)
+	if (PublicKey != InvalidRSAKeyHandle)
 	{
-		SignaturesFile->Serialize(*Reader);
-		delete Reader;
+		FString SignaturesFilename = FPaths::ChangeExtension(InFilename, TEXT("sig"));
+		FArchive* Reader = IFileManager::Get().CreateFileReader(*SignaturesFilename);
+		if (Reader != nullptr)
+		{
+			SignaturesFile = MakeShared<FPakSignatureFile, ESPMode::ThreadSafe>();
+			SignaturesFile->Serialize(*Reader);
+			delete Reader;
+
+			if (!SignaturesFile->DecryptSignatureAndValidate(PublicKey, InFilename))
+			{
+				// We don't need to act on this failure as the decrypt function will already have dumped out log messages
+				// and fired the signature check fail handler
+				SignaturesFile.Reset();
+			}
 
-		if (!SignaturesFile->DecryptSignatureAndValidate(InFilename))
+			PakSignatureFileCache.Add(FilenameFName, SignaturesFile);
+		}
+		else
 		{
-			SignaturesFile.Reset();
+			UE_LOG(LogPakFile, Warning, TEXT("Couldn't find pak signature file '%s'"), InFilename);
+			FPakPlatformFile::GetPakMasterSignatureTableCheckFailureHandler().Broadcast(InFilename);
 		}
-
-		PakSignatureFileCache.Add(FilenameFName, SignaturesFile);
-	}
-	else
-	{
-		UE_LOG(LogPakFile, Warning, TEXT("Couldn't find pak signature file '%s'"), InFilename);
-		FPakPlatformFile::GetPakMasterSignatureTableCheckFailureHandler().Broadcast(InFilename);
 	}
 
 	return SignaturesFile;
@@ -5933,9 +5934,9 @@ bool FPakPlatformFile::Initialize(IPlatformFile* Inner, const TCHAR* CmdLine)
 	GameUserSettingsIniFilename = TEXT("GameUserSettings.ini");
 #endif
 
-	// signed if we have keys, and are not running with fileopenlog (currently results in a deadlock).
-	bSigned = GetPakSigningKey().IsValid() && !FParse::Param(FCommandLine::Get(), TEXT("fileopenlog"));;
-		
+	// Signed if we have keys, and are not running with fileopenlog (currently results in a deadlock).
+	bSigned = FCoreDelegates::GetPakSigningKeysDelegate().IsBound() && !FParse::Param(FCommandLine::Get(), TEXT("fileopenlog"));
+
 	// Find and mount pak files from the specified directories.
 	TArray<FString> PakFolders;
 	GetPakFolders(FCommandLine::Get(), PakFolders);
