Fix bad merge in IPlatformFilePak.cpp
#rnx

diff --git a/Engine/Source/Runtime/PakFile/Private/IPlatformFilePak.cpp b/Engine/Source/Runtime/PakFile/Private/IPlatformFilePak.cpp
--- a/Engine/Source/Runtime/PakFile/Private/IPlatformFilePak.cpp
+++ b/Engine/Source/Runtime/PakFile/Private/IPlatformFilePak.cpp
@@ -1249,16 +1249,16 @@ class FPakPrecacher
 	uint32 Loads;
 	uint32 Frees;
 	uint64 LoadSize;
-	FRSA::TKeyPtr SigningKey;
 	EAsyncIOPriorityAndFlags AsyncMinPriority;
 	FCriticalSection SetAsyncMinimumPriorityScopeLock;
+	bool bEnableSignatureChecks;
 public:
 
-	static void Init(IPlatformFile* InLowerLevel, FRSA::TKeyPtr InSigningKey)
+	static void Init(IPlatformFile* InLowerLevel, bool bInEnableSignatureChecks) 
 	{
 		if (!PakPrecacherSingleton)
 		{
-			verify(!FPlatformAtomics::InterlockedCompareExchangePointer((void**)&PakPrecacherSingleton, new FPakPrecacher(InLowerLevel, InSigningKey), nullptr));
+			verify(!FPlatformAtomics::InterlockedCompareExchangePointer((void**)& PakPrecacherSingleton, new FPakPrecacher(InLowerLevel, bInEnableSignatureChecks), nullptr));
 		}
 		check(PakPrecacherSingleton);
 	}
@@ -1294,7 +1294,7 @@ public:
 		return *PakPrecacherSingleton;
 	}
 
-	FPakPrecacher(IPlatformFile* InLowerLevel, FRSA::TKeyPtr InSigningKey)
+	FPakPrecacher(IPlatformFile* InLowerLevel, bool bInEnableSignatureChecks) 
 		: LowerLevel(InLowerLevel)
 		, LastReadRequest(0)
 		, NextUniqueID(1)
@@ -1304,8 +1304,8 @@ public:
 		, Loads(0)
 		, Frees(0)
 		, LoadSize(0)
-		, SigningKey(InSigningKey)
 		, AsyncMinPriority(AIOP_MIN)
+		, bEnableSignatureChecks(bInEnableSignatureChecks)
 	{
 		check(LowerLevel && FPlatformProcess::SupportsMultithreading());
 		GPakCache_MaxRequestsToLowerLevel = FMath::Max(FMath::Min(FPlatformMisc::NumberOfIOWorkerThreadsToSpawn(), GPakCache_MaxRequestsToLowerLevel), 1);
@@ -1378,7 +1378,7 @@ public:
 
 			UE_LOG(LogPakFile, Log, TEXT("New pak file %s added to pak precacher."), *PakFilename);
 
-			if (SigningKey.IsValid())
+			if (Pak.Signatures.IsValid())
 			{
 				// Load signature data
 				Pak.Signatures = FPakPlatformFile::GetPakSignatureFile(*PakFilename);
@@ -2307,7 +2307,7 @@ private: // below here we assume CachedFilesScopeLock until we get to the next s
 		FAsyncFileCallBack CallbackFromLower =
 			[this, IndexToFill, bDoCheck](bool bWasCanceled, IAsyncReadRequest* Request)
 		{
-			if (SigningKey.IsValid() && bDoCheck)
+			if (bEnableSignatureChecks && bDoCheck)
 			{
 				StartSignatureCheck(bWasCanceled, Request, IndexToFill);
 			}
@@ -5962,7 +5962,7 @@ void FPakPlatformFile::InitializeNewAsyncIO()
 #if !WITH_EDITOR
 	if (FPlatformProcess::SupportsMultithreading() && !FParse::Param(FCommandLine::Get(), TEXT("FileOpenLog")))
 	{
-		FPakPrecacher::Init(LowerLevel, GetPakSigningKey());
+		FPakPrecacher::Init(LowerLevel, FCoreDelegates::GetPakSigningKeysDelegate().IsBound());
 	}
 	else
 #endif
