Add a per frame KB stat for the pre cache.

diff --git a/Engine/Source/Runtime/PakFile/Private/IPlatformFilePak.cpp b/Engine/Source/Runtime/PakFile/Private/IPlatformFilePak.cpp
--- a/Engine/Source/Runtime/PakFile/Private/IPlatformFilePak.cpp
+++ b/Engine/Source/Runtime/PakFile/Private/IPlatformFilePak.cpp
@@ -352,6 +352,7 @@ DECLARE_MEMORY_STAT(TEXT("PakCache High Water"), STAT_PakCacheHighWater, STATGRO
 volatile int64 GPreCacheHotBlocksCount = 0;
 volatile int64 GPreCacheColdBlocksCount = 0;
 volatile int64 GPreCacheTotalLoaded = 0;
+int64 GPreCacheTotalLoadedLastTick = 0;
 #endif
 
 DECLARE_FLOAT_ACCUMULATOR_STAT(TEXT("PakCache Signing Chunk Hash Time"), STAT_PakCache_SigningChunkHashTime, STATGROUP_PakFile);
@@ -4101,6 +4102,14 @@ void FPakPlatformFile::Tick()
 		CSV_CUSTOM_STAT(FileIO, PakPrecacherHotBlocksCount, (int32)GPreCacheHotBlocksCount, ECsvCustomStatOp::Set);
 		CSV_CUSTOM_STAT(FileIO, PakPrecacherColdBlocksCount, (int32)GPreCacheColdBlocksCount, ECsvCustomStatOp::Set);
 		CSV_CUSTOM_STAT(FileIO, PakPrecacherTotalLoadedMB, (int32)(GPreCacheTotalLoaded/(1024*1024)), ECsvCustomStatOp::Set);
+
+		if (GPreCacheTotalLoadedLastTick != 0)
+		{
+			int64 diff = GPreCacheTotalLoaded - GPreCacheTotalLoadedLastTick;
+			diff /= 1024;
+			CSV_CUSTOM_STAT(FileIO, PakPrecacherPerFrameKB, (int32)diff, ECsvCustomStatOp::Set);
+		}
+		GPreCacheTotalLoadedLastTick = GPreCacheTotalLoaded;
 }
 #endif
 }
