Add a delegate which fires when the startup loading screen is shown
Also add support for overriding the startup pak mount wildcard

diff --git a/Engine/Source/Runtime/Core/Private/Misc/CoreDelegates.cpp b/Engine/Source/Runtime/Core/Private/Misc/CoreDelegates.cpp
--- a/Engine/Source/Runtime/Core/Private/Misc/CoreDelegates.cpp
+++ b/Engine/Source/Runtime/Core/Private/Misc/CoreDelegates.cpp
@@ -48,6 +48,8 @@ FCoreDelegates::FOnActorLabelChanged FCoreDelegates::OnActorLabelChanged;
 FCoreDelegates::FRegisterMovieStreamerDelegate FCoreDelegates::RegisterMovieStreamerDelegate;
 FCoreDelegates::FUnRegisterMovieStreamerDelegate FCoreDelegates::UnRegisterMovieStreamerDelegate;
 
+FCoreDelegates::FOnInitialLoadingScreenShown FCoreDelegates::OnInitialLoadingScreenShown;
+
 FCoreDelegates::FRegisterEncryptionKeyDelegate& FCoreDelegates::GetRegisterEncryptionKeyDelegate()
 {
 	static FRegisterEncryptionKeyDelegate RegisterEncryptionKeyDelegate;
diff --git a/Engine/Source/Runtime/Core/Public/Misc/CoreDelegates.h b/Engine/Source/Runtime/Core/Public/Misc/CoreDelegates.h
--- a/Engine/Source/Runtime/Core/Public/Misc/CoreDelegates.h
+++ b/Engine/Source/Runtime/Core/Public/Misc/CoreDelegates.h
@@ -131,6 +131,10 @@ public:
 	DECLARE_MULTICAST_DELEGATE_OneParam(FOnSyncLoadPackage, const FString&);
 	static FOnSyncLoadPackage OnSyncLoadPackage;
 
+	// Callback for when the initial loading screen is shown during PreInit
+	DECLARE_MULTICAST_DELEGATE(FOnInitialLoadingScreenShown);
+	static FOnInitialLoadingScreenShown OnInitialLoadingScreenShown;
+
 	// get a hotfix delegate
 	static FHotFixDelegate& GetHotfixDelegate(EHotfixDelegates::Type HotFix);
 
diff --git a/Engine/Source/Runtime/Launch/Private/LaunchEngineLoop.cpp b/Engine/Source/Runtime/Launch/Private/LaunchEngineLoop.cpp
--- a/Engine/Source/Runtime/Launch/Private/LaunchEngineLoop.cpp
+++ b/Engine/Source/Runtime/Launch/Private/LaunchEngineLoop.cpp
@@ -2534,6 +2534,11 @@ int32 FEngineLoop::PreInit(const TCHAR* CmdLine)
 					FPlatformMisc::PlatformHandleSplashScreen(true);
 				}
             }
+		
+			if (FCoreDelegates::OnInitialLoadingScreenShown.IsBound())
+			{
+				FCoreDelegates::OnInitialLoadingScreenShown.Broadcast();
+			}
 		}
 		else if ( IsRunningCommandlet() )
 		{
diff --git a/Engine/Source/Runtime/PakFile/Private/IPlatformFilePak.cpp b/Engine/Source/Runtime/PakFile/Private/IPlatformFilePak.cpp
--- a/Engine/Source/Runtime/PakFile/Private/IPlatformFilePak.cpp
+++ b/Engine/Source/Runtime/PakFile/Private/IPlatformFilePak.cpp
@@ -51,6 +51,8 @@ CSV_DECLARE_CATEGORY_MODULE_EXTERN(CORE_API, FileIO);
 #define MOUNT_STARTUP_PAKS_WILDCARD ALL_PAKS_WILDCARD
 #endif
 
+static FString GMountStartupPaksWildCard = TEXT(MOUNT_STARTUP_PAKS_WILDCARD);
+
 int32 ParseChunkIDFromFilename(const FString& InFilename)
 {
 	FString ChunkIdentifier(TEXT("pakchunk"));
@@ -5967,10 +5969,15 @@ bool FPakPlatformFile::Initialize(IPlatformFile* Inner, const TCHAR* CmdLine)
 	// Signed if we have keys, and are not running with fileopenlog (currently results in a deadlock).
 	bSigned = FCoreDelegates::GetPakSigningKeysDelegate().IsBound() && !FParse::Param(FCommandLine::Get(), TEXT("fileopenlog"));
 
+	FString StartupPaksWildcard = GMountStartupPaksWildCard;
+#if !UE_BUILD_SHIPPING
+	FParse::Value(FCommandLine::Get(), TEXT("StartupPaksWildcard="), StartupPaksWildcard);
+#endif
+
 	// Find and mount pak files from the specified directories.
 	TArray<FString> PakFolders;
 	GetPakFolders(FCommandLine::Get(), PakFolders);
-	MountAllPakFiles(PakFolders, TEXT(MOUNT_STARTUP_PAKS_WILDCARD));
+	MountAllPakFiles(PakFolders, *StartupPaksWildcard);
 
 #if !UE_BUILD_SHIPPING
 	GPakExec = MakeUnique<FPakExec>(*this);
@@ -6426,9 +6433,15 @@ IFileHandle* FPakPlatformFile::OpenRead(const TCHAR* Filename, bool bAllowWrite)
 
 const TCHAR* FPakPlatformFile::GetMountStartupPaksWildCard()
 {
-	return TEXT(MOUNT_STARTUP_PAKS_WILDCARD);
+	return *GMountStartupPaksWildCard;
 }
 
+void FPakPlatformFile::SetMountStartupPaksWildCard(const FString& WildCard)
+{
+	GMountStartupPaksWildCard = WildCard;
+}
+
+
 EChunkLocation::Type FPakPlatformFile::GetPakChunkLocation(int32 InChunkID) const
 {
 	FScopeLock ScopedLock(&PakListCritical);
diff --git a/Engine/Source/Runtime/PakFile/Public/IPlatformFilePak.h b/Engine/Source/Runtime/PakFile/Public/IPlatformFilePak.h
--- a/Engine/Source/Runtime/PakFile/Public/IPlatformFilePak.h
+++ b/Engine/Source/Runtime/PakFile/Public/IPlatformFilePak.h
@@ -1560,6 +1560,11 @@ public:
 	 */
 	static const TCHAR* GetMountStartupPaksWildCard();
 
+	/**
+	 * Overrides the wildcard used for searching paks. Call before initialization
+	 */
+	static void SetMountStartupPaksWildCard(const FString& WildCard);
+
 	/**
 	* Determine location information for a given chunk ID. Will be DoesNotExist if the pak file wasn't detected, NotAvailable if it exists but hasn't been mounted due to a missing encryption key, or LocalFast if it exists and has been mounted
 	*/
