Compile fix (in UE4Editor non unity Win64)

diff --git a/Engine/Source/Runtime/PakFile/Private/IPlatformFilePak.cpp b/Engine/Source/Runtime/PakFile/Private/IPlatformFilePak.cpp
--- a/Engine/Source/Runtime/PakFile/Private/IPlatformFilePak.cpp
+++ b/Engine/Source/Runtime/PakFile/Private/IPlatformFilePak.cpp
@@ -42,7 +42,9 @@ CSV_DECLARE_CATEGORY_MODULE_EXTERN(CORE_API, FileIO);
 #ifndef DISABLE_NONUFS_INI_WHEN_COOKED
 #define DISABLE_NONUFS_INI_WHEN_COOKED 0
 #endif
-
+#ifndef ALLOW_INI_OVERRIDE_FROM_COMMANDLINE
+#define ALLOW_INI_OVERRIDE_FROM_COMMANDLINE 0
+#endif
 #ifndef ALL_PAKS_WILDCARD
 #define ALL_PAKS_WILDCARD "*.pak"
 #endif 
