Disable filename extension exclusions when using pak files in an editor build

diff --git a/Engine/Source/Runtime/PakFile/Private/IPlatformFilePak.cpp b/Engine/Source/Runtime/PakFile/Private/IPlatformFilePak.cpp
--- a/Engine/Source/Runtime/PakFile/Private/IPlatformFilePak.cpp
+++ b/Engine/Source/Runtime/PakFile/Private/IPlatformFilePak.cpp
@@ -6070,7 +6070,7 @@ bool FPakPlatformFile::Initialize(IPlatformFile* Inner, const TCHAR* CmdLine)
 	check(Inner != NULL);
 	LowerLevel = Inner;
 
-#if EXCLUDE_NONPAK_UE_EXTENSIONS
+#if EXCLUDE_NONPAK_UE_EXTENSIONS && !WITH_EDITOR
 	// Extensions for file types that should only ever be in a pak file. Used to stop unnecessary access to the lower level platform file
 	ExcludedNonPakExtensions.Add(TEXT("uasset"));
 	ExcludedNonPakExtensions.Add(TEXT("umap"));
