Pak platform file changes to support mounting pak files in the editor
 - Always look for pak files in the standard locations to determine whether we should create the platform layer
 - When looking up a decryption key, check the registered list for all guids, even empty ones. We want to support pak mounting in non-monolithic builds where we don't have an embedded key.
 - Remove the initialization-time check that the decryption key exists for pak files with an encrypted index. The condition to test is more complex when considering editor pak mounting, and we will get a meaningful error almost immediately afterwards anyway.

diff --git a/Engine/Source/Runtime/PakFile/Private/IPlatformFilePak.cpp b/Engine/Source/Runtime/PakFile/Private/IPlatformFilePak.cpp
--- a/Engine/Source/Runtime/PakFile/Private/IPlatformFilePak.cpp
+++ b/Engine/Source/Runtime/PakFile/Private/IPlatformFilePak.cpp
@@ -235,17 +235,17 @@ void FPakPlatformFile::GetPakEncryptionKey(FAES::FAESKey& OutKey, const FGuid& I
 {
 	OutKey.Reset();
 
-	if (InEncryptionKeyGuid.IsValid())
+	if (!GetRegisteredEncryptionKeys().GetKey(InEncryptionKeyGuid, OutKey))
 	{
-		if (!GetRegisteredEncryptionKeys().GetKey(InEncryptionKeyGuid, OutKey))
+		if (!InEncryptionKeyGuid.IsValid() && FCoreDelegates::GetPakEncryptionKeyDelegate().IsBound())
+		{
+			FCoreDelegates::GetPakEncryptionKeyDelegate().Execute(OutKey.Key);
+		}
+		else
 		{
 			UE_LOG(LogPakFile, Fatal, TEXT("Failed to find requested encryption key %s"), *InEncryptionKeyGuid.ToString());
 		}
 	}
-	else
-	{
-		FCoreDelegates::GetPakEncryptionKeyDelegate().ExecuteIfBound(OutKey.Key);
-	}
 }
 
 TMap<FName, TSharedPtr<const struct FPakSignatureFile, ESPMode::ThreadSafe>> FPakPlatformFile::PakSignatureFileCache;
@@ -4811,7 +4811,6 @@ void FPakFile::Initialize(FArchive* Reader, bool bLoadIndex)
 	{
 		UE_CLOG(Info.Magic != FPakInfo::PakFile_Magic, LogPakFile, Fatal, TEXT("Trailing magic number (%ud) in '%s' is different than the expected one. Verify your installation."), Info.Magic, *PakFilename);
 		UE_CLOG(!(Info.Version >= FPakInfo::PakFile_Version_Initial && Info.Version <= CompatibleVersion), LogPakFile, Fatal, TEXT("Invalid pak file version (%d) in '%s'. Verify your installation."), Info.Version, *PakFilename);
-		UE_CLOG((Info.bEncryptedIndex == 1) && (!FCoreDelegates::GetPakEncryptionKeyDelegate().IsBound()), LogPakFile, Fatal, TEXT("Index of pak file '%s' is encrypted, but this executable doesn't have any valid decryption keys"), *PakFilename);
 		UE_CLOG(!(Info.IndexOffset >= 0 && Info.IndexOffset < CachedTotalSize), LogPakFile, Fatal, TEXT("Index offset for pak file '%s' is invalid (%lld)"), *PakFilename, Info.IndexOffset);
 		UE_CLOG(!((Info.IndexOffset + Info.IndexSize) >= 0 && (Info.IndexOffset + Info.IndexSize) <= CachedTotalSize), LogPakFile, Fatal, TEXT("Index end offset for pak file '%s' is invalid (%lld)"), *PakFilename, Info.IndexOffset + Info.IndexSize);
 
@@ -5920,10 +5919,6 @@ FPakPlatformFile::FPakPlatformFile()
 	, bSigned(false)
 {
 	FCoreDelegates::GetRegisterEncryptionKeyDelegate().BindRaw(this, &FPakPlatformFile::RegisterEncryptionKey);
-
-	// Register an empty guid against an empty key. An empty guid means use the embedded AES key, which will be looked up dynamically on request. This is done for data hiding purposes, but
-	// if we decide that there is no point protecting the embedded key, we could cache it here for speed purposes.
-	RegisterEncryptionKey(FGuid(), FAES::FAESKey());
 }
 
 FPakPlatformFile::~FPakPlatformFile()
@@ -6051,14 +6046,12 @@ bool FPakPlatformFile::CheckIfPakFilesExist(IPlatformFile* LowLevelFile, const T
 bool FPakPlatformFile::ShouldBeUsed(IPlatformFile* Inner, const TCHAR* CmdLine) const
 {
 	bool Result = false;
-#if !WITH_EDITOR
 	if (!FParse::Param(CmdLine, TEXT("NoPak")))
 	{
 		TArray<FString> PakFolders;
 		GetPakFolders(CmdLine, PakFolders);
 		Result = CheckIfPakFilesExist(Inner, PakFolders);
 	}
-#endif
 	return Result;
 }
 
@@ -6171,7 +6164,7 @@ bool FPakPlatformFile::Mount(const TCHAR* InPakFilename, uint32 PakOrder, const
 		FPakFile* Pak = new FPakFile(LowerLevel, InPakFilename, bSigned, bLoadIndex);
 		if (Pak->IsValid())
 		{
-			if (GetRegisteredEncryptionKeys().HasKey(Pak->GetInfo().EncryptionKeyGuid))
+			if (!Pak->GetInfo().EncryptionKeyGuid.IsValid() || GetRegisteredEncryptionKeys().HasKey(Pak->GetInfo().EncryptionKeyGuid))
 			{
 				if (InPath != NULL)
 				{
