Fix crash when trying to load files from PAKs in single threaded applications

diff --git a/Engine/Source/Runtime/PakFile/Private/IPlatformFilePak.cpp b/Engine/Source/Runtime/PakFile/Private/IPlatformFilePak.cpp
--- a/Engine/Source/Runtime/PakFile/Private/IPlatformFilePak.cpp
+++ b/Engine/Source/Runtime/PakFile/Private/IPlatformFilePak.cpp
@@ -4161,7 +4161,7 @@ IAsyncReadFileHandle* FPakPlatformFile::OpenAsyncRead(const TCHAR* Filename)
 		}
 	}
 #endif
-	return LowerLevel->OpenAsyncRead(Filename);
+	return IPlatformFile::OpenAsyncRead(Filename);
 }
 
 void FPakPlatformFile::SetAsyncMinimumPriority(EAsyncIOPriorityAndFlags Priority)
