Only create the pak platform layer in editor builds if we're in a monolithic builds. A modular build won't have any required decryption keys embedded, so this functionality will be disabled until we have a way of specifying a keychain and loading it.

diff --git a/Engine/Source/Runtime/PakFile/Private/IPlatformFilePak.cpp b/Engine/Source/Runtime/PakFile/Private/IPlatformFilePak.cpp
--- a/Engine/Source/Runtime/PakFile/Private/IPlatformFilePak.cpp
+++ b/Engine/Source/Runtime/PakFile/Private/IPlatformFilePak.cpp
@@ -6046,7 +6046,7 @@ bool FPakPlatformFile::CheckIfPakFilesExist(IPlatformFile* LowLevelFile, const T
 bool FPakPlatformFile::ShouldBeUsed(IPlatformFile* Inner, const TCHAR* CmdLine) const
 {
 	bool Result = false;
-	if (!FParse::Param(CmdLine, TEXT("NoPak")))
+	if ((!WITH_EDITOR || IS_MONOLITHIC) && !FParse::Param(CmdLine, TEXT("NoPak")))
 	{
 		TArray<FString> PakFolders;
 		GetPakFolders(CmdLine, PakFolders);
