Don't print out encryption key registration log for empty guids

diff --git a/Engine/Source/Runtime/PakFile/Private/IPlatformFilePak.cpp b/Engine/Source/Runtime/PakFile/Private/IPlatformFilePak.cpp
--- a/Engine/Source/Runtime/PakFile/Private/IPlatformFilePak.cpp
+++ b/Engine/Source/Runtime/PakFile/Private/IPlatformFilePak.cpp
@@ -5928,7 +5928,7 @@ void FPakPlatformFile::RegisterEncryptionKey(const FGuid& InGuid, const FAES::FA
 		PendingEncryptedPakFiles.RemoveAll([InGuid](const FPakListDeferredEntry& Entry) { return Entry.EncryptionKeyGuid == InGuid; });
 	}
 
-	UE_LOG(LogPakFile, Log, TEXT("Registered encryption key '%s': %d pak files mounted, %d remain pending"), *InGuid.ToString(), NumMounted, PendingEncryptedPakFiles.Num());
+	UE_CLOG(InGuid.IsValid(), LogPakFile, Log, TEXT("Registered encryption key '%s': %d pak files mounted, %d remain pending"), *InGuid.ToString(), NumMounted, PendingEncryptedPakFiles.Num());
 }
 
 IFileHandle* FPakPlatformFile::OpenRead(const TCHAR* Filename, bool bAllowWrite)
