PR #6056: Removed int with UE preferred int32 type (Contributed by [REDACTED_PERSON])

diff --git a/Engine/Source/Runtime/PakFile/Public/IPlatformFilePak.h b/Engine/Source/Runtime/PakFile/Public/IPlatformFilePak.h
--- a/Engine/Source/Runtime/PakFile/Public/IPlatformFilePak.h
+++ b/Engine/Source/Runtime/PakFile/Public/IPlatformFilePak.h
@@ -211,7 +211,7 @@ struct FPakInfo
 			if (Ar.IsLoading())
 			{
 				Ar.Serialize(Methods, BufferSize);
-				for (int Index = 0; Index < MaxNumCompressionMethods; Index++)
+				for (int32 Index = 0; Index < MaxNumCompressionMethods; Index++)
 				{
 					ANSICHAR* MethodString = &Methods[Index * CompressionMethodNameLen];
 					if (MethodString[0] != 0)
@@ -225,7 +225,7 @@ struct FPakInfo
 				// we always zero out fully what we write out so that reading in is simple
 				FMemory::Memzero(Methods, BufferSize);
 
-				for (int Index = 1; Index < CompressionMethods.Num(); Index++)
+				for (int32 Index = 1; Index < CompressionMethods.Num(); Index++)
 				{
 					ANSICHAR* MethodString = &Methods[(Index - 1) * CompressionMethodNameLen];
 					FCStringAnsi::Strcpy(MethodString, CompressionMethodNameLen, TCHAR_TO_ANSI(*CompressionMethods[Index].ToString()));
@@ -2506,4 +2506,4 @@ struct FPakSignatureFile
 		FSHA1::HashBuffer(ChunkHashes.GetData(), ChunkHashes.Num() * sizeof(TPakChunkHash), CurrentHash.Hash);
 		return CurrentHash;
 	}
-};
\ No newline at end of file
+};
