Moved some macros around to avoid a static analyser issue

#rnx

diff --git a/Engine/Source/Runtime/PakFile/Private/IPlatformFilePak.cpp b/Engine/Source/Runtime/PakFile/Private/IPlatformFilePak.cpp
--- a/Engine/Source/Runtime/PakFile/Private/IPlatformFilePak.cpp
+++ b/Engine/Source/Runtime/PakFile/Private/IPlatformFilePak.cpp
@@ -6046,12 +6046,14 @@ bool FPakPlatformFile::CheckIfPakFilesExist(IPlatformFile* LowLevelFile, const T
 bool FPakPlatformFile::ShouldBeUsed(IPlatformFile* Inner, const TCHAR* CmdLine) const
 {
 	bool Result = false;
-	if ((!WITH_EDITOR || IS_MONOLITHIC) && !FParse::Param(CmdLine, TEXT("NoPak")))
+#if (!WITH_EDITOR || IS_MONOLITHIC)
+	if (!FParse::Param(CmdLine, TEXT("NoPak")))
 	{
 		TArray<FString> PakFolders;
 		GetPakFolders(CmdLine, PakFolders);
 		Result = CheckIfPakFilesExist(Inner, PakFolders);
 	}
+#endif
 	return Result;
 }
 
