PakCache trim can be based on time
Add some more stats for detecting bad seeks between pak files and contiguous reads

diff --git a/Engine/Source/Runtime/PakFile/Private/IPlatformFilePak.cpp b/Engine/Source/Runtime/PakFile/Private/IPlatformFilePak.cpp
--- a/Engine/Source/Runtime/PakFile/Private/IPlatformFilePak.cpp
+++ b/Engine/Source/Runtime/PakFile/Private/IPlatformFilePak.cpp
@@ -379,6 +379,10 @@ volatile int64 GPreCacheHotBlocksCount = 0;
 volatile int64 GPreCacheColdBlocksCount = 0;
 volatile int64 GPreCacheTotalLoaded = 0;
 int64 GPreCacheTotalLoadedLastTick = 0;
+
+volatile int64 GPreCacheSeeks = 0;
+volatile int64 GPreCacheBadSeeks = 0;
+volatile int64 GPreCacheContiguousReads = 0;
 #endif
 
 DECLARE_FLOAT_ACCUMULATOR_STAT(TEXT("PakCache Signing Chunk Hash Time"), STAT_PakCache_SigningChunkHashTime, STATGROUP_PakFile);
@@ -435,6 +439,13 @@ static FAutoConsoleVariableRef CVar_NumUnreferencedBlocksToCache(
 	TEXT("Controls the maximum number of unreferenced blocks to keep. This is a classic disk cache and the maxmimum wasted memory is pakcache.MaxRequestSizeToLowerLevellKB * pakcache.NumUnreferencedBlocksToCache.")
 );
 
+float GPakCache_TimeToTrim = 0.0f;
+static FAutoConsoleVariableRef CVar_PakCache_TimeToTrim(
+	TEXT("pakcache.TimeToTrim"),
+	GPakCache_TimeToTrim,
+	TEXT("Controls how long to hold onto a cached but unreferenced block for.")
+);
+
 class FPakPrecacher;
 
 typedef uint64 FJoinedOffsetAndPakIndex;
@@ -1135,6 +1146,7 @@ class FPakPrecacher
 		TIntervalTreeIndex Index;
 		TIntervalTreeIndex Next;
 		EBlockStatus Status;
+		double TimeNoLongerReferenced;
 
 		FCacheBlock()
 			: OffsetAndPakIndex(0)
@@ -1286,6 +1298,8 @@ class FPakPrecacher
 	FCriticalSection SetAsyncMinimumPriorityScopeLock;
 	bool bEnableSignatureChecks;
 public:
+	int64 GetBlockMemory() { return BlockMemory; }
+	int64 GetBlockMemoryHighWater() { return BlockMemoryHighWater; }
 
 	static void Init(IPlatformFile* InLowerLevel, bool bInEnableSignatureChecks) 
 	{
@@ -1341,7 +1355,7 @@ public:
 		, bEnableSignatureChecks(bInEnableSignatureChecks)
 	{
 		check(LowerLevel && FPlatformProcess::SupportsMultithreading());
-		GPakCache_MaxRequestsToLowerLevel = FMath::Max(FMath::Min(FPlatformMisc::NumberOfIOWorkerThreadsToSpawn(), GPakCache_MaxRequestsToLowerLevel), 1);
+//		GPakCache_MaxRequestsToLowerLevel = FMath::Max(FMath::Min(FPlatformMisc::NumberOfIOWorkerThreadsToSpawn(), GPakCache_MaxRequestsToLowerLevel), 1);
 		check(GPakCache_MaxRequestsToLowerLevel <= PAK_CACHE_MAX_REQUESTS);
 	}
 
@@ -1536,10 +1550,9 @@ private: // below here we assume CachedFilesScopeLock until we get to the next s
 		return MAX_uint64;
 	}
 
-	bool AddRequest(TIntervalTreeIndex NewIndex)
+	bool AddRequest(FPakInRequest& Request, TIntervalTreeIndex NewIndex)
 	{
 		// CachedFilesScopeLock is locked
-		FPakInRequest& Request = InRequestAllocator.Get(NewIndex);
 		uint16 PakIndex = GetRequestPakIndex(Request.OffsetAndPakIndex);
 		int64 Offset = GetRequestOffset(Request.OffsetAndPakIndex);
 		FPakData& Pak = CachedPakData[PakIndex];
@@ -1702,7 +1715,7 @@ private: // below here we assume CachedFilesScopeLock until we get to the next s
 	void TrimCache(bool bDiscardAll = false, uint16 StartPakIndex = 65535)
 	{
 
-		if (GPakCache_UseNewTrim)
+		if (GPakCache_UseNewTrim && !bDiscardAll)
 		{
 			StartPakIndex = 0;
 			uint16 EndPakIndex = CachedPakData.Num();
@@ -1712,7 +1725,7 @@ private: // below here we assume CachedFilesScopeLock until we get to the next s
 			CacheVisitedAlready.AddDefaulted(OffsetAndPakIndexOfSavedBlocked.Num());
 
 			int64 MemoryBudget = GPakCache_MaxBlockMemory * (1024 * 1024);
-
+			bool AlreadyRemovedBlocksBecauseOfMemoryOverage = false;
 
 			while (BlockMemory > MemoryBudget)
 			{
@@ -1767,6 +1780,7 @@ private: // below here we assume CachedFilesScopeLock until we get to the next s
 							);
 						}
 						OffsetAndPakIndexOfSavedBlocked[CacheIndex].RemoveAt(0, NumToRemove, false);
+						AlreadyRemovedBlocksBecauseOfMemoryOverage = true;
 					}
 				}
 				if (NoneToRemove)
@@ -1774,6 +1788,59 @@ private: // below here we assume CachedFilesScopeLock until we get to the next s
 					break;
 				}
 			}
+			if(GPakCache_TimeToTrim != 0.0f)
+			{
+				// we'll trim based on time rather than trying to keep within a memory budget
+				double CurrentTime = FPlatformTime::Seconds();
+				// CachedFilesScopeLock is locked
+				for (uint16 RealPakIndex = StartPakIndex; RealPakIndex < EndPakIndex; RealPakIndex++)
+				{
+					int32 CacheIndex = CachedPakData[RealPakIndex].ActualPakFile->GetCacheIndex();
+
+					int32 NumToRemove = 0;
+
+					if (OffsetAndPakIndexOfSavedBlocked[CacheIndex].Num())
+					{
+						for (int32 Index = 0; Index < OffsetAndPakIndexOfSavedBlocked[CacheIndex].Num(); Index++)
+						{
+							FJoinedOffsetAndPakIndex OffsetAndPakIndex = OffsetAndPakIndexOfSavedBlocked[CacheIndex][Index];
+							uint16 PakIndex = GetRequestPakIndex(OffsetAndPakIndex);
+							int64 Offset = GetRequestOffset(OffsetAndPakIndex);
+							FPakData& Pak = CachedPakData[PakIndex];
+							bool removedall = true;
+							MaybeRemoveOverlappingNodesInIntervalTree<FCacheBlock>(
+								&Pak.CacheBlocks[(int32)EBlockStatus::Complete],
+								CacheBlockAllocator,
+								Offset,
+								Offset,
+								0,
+								Pak.MaxNode,
+								Pak.StartShift,
+								Pak.MaxShift,
+								[this, CurrentTime, &removedall](TIntervalTreeIndex BlockIndex) -> bool
+							{
+								FCacheBlock &Block = CacheBlockAllocator.Get(BlockIndex);
+								if (!Block.InRequestRefCount && (CurrentTime - Block.TimeNoLongerReferenced >= GPakCache_TimeToTrim))
+								{
+									UE_LOG(LogPakFile, Verbose, TEXT("FPakReadRequest[%016llX, %016llX) Discard Cached Based on Time"), Block.OffsetAndPakIndex, Block.OffsetAndPakIndex + Block.Size);
+									ClearBlock(Block);
+									return true;
+								}
+								removedall = false;
+								return false;
+							}
+							);
+							if (!removedall)
+								break;
+							NumToRemove++;
+						}
+						if (NumToRemove)
+						{
+							OffsetAndPakIndexOfSavedBlocked[CacheIndex].RemoveAt(0, NumToRemove, false);
+						}
+					}
+				}
+			}
 		}
 		else
 		{
@@ -1863,7 +1930,7 @@ private: // below here we assume CachedFilesScopeLock until we get to the next s
 					{
 						uint16 BlocksPakIndex = GetRequestPakIndexLow(Block.OffsetAndPakIndex);
 						int32 BlocksCacheIndex = CachedPakData[BlocksPakIndex].ActualPakFile->GetCacheIndex();
-
+						Block.TimeNoLongerReferenced = FPlatformTime::Seconds();
 						OffsetAndPakIndexOfSavedBlocked[BlocksCacheIndex].Remove(Block.OffsetAndPakIndex);
 						OffsetAndPakIndexOfSavedBlocked[BlocksCacheIndex].Add(Block.OffsetAndPakIndex);
 						return false;
@@ -2356,7 +2423,26 @@ private: // below here we assume CachedFilesScopeLock until we get to the next s
 
 		RequestsToLower[IndexToFill].RequestHandle = Pak.Handle->ReadRequest(GetRequestOffset(Block.OffsetAndPakIndex), Block.Size, Priority, &CallbackFromLower);
 		RedundantReadTracker.CheckBlock(GetRequestOffset(Block.OffsetAndPakIndex), Block.Size);
+
+		FJoinedOffsetAndPakIndex OldLastReadRequest = LastReadRequest;
 		LastReadRequest = Block.OffsetAndPakIndex + Block.Size;
+
+		if (OldLastReadRequest != Block.OffsetAndPakIndex)
+		{
+			if (GetRequestPakIndexLow(OldLastReadRequest) != GetRequestPakIndexLow(Block.OffsetAndPakIndex))
+			{
+				GPreCacheBadSeeks++;
+			}
+			else
+			{
+				GPreCacheSeeks++;
+			}
+		}
+		else
+		{
+			GPreCacheContiguousReads++;
+		}
+
 		Loads++;
 		LoadSize += Block.Size;
 	}
@@ -2603,7 +2689,7 @@ public:
 		OutstandingRequests.Add(Request.UniqueID, RequestIndex);
 		RequestCounter.Increment();
 
-		if (AddRequest(RequestIndex))
+		if (AddRequest(Request, RequestIndex))
 		{
 #if USE_PAK_PRECACHE && CSV_PROFILER
 			FPlatformAtomics::InterlockedIncrement(&GPreCacheHotBlocksCount);
@@ -4185,6 +4271,8 @@ void FPakPlatformFile::Tick()
 		CSV_CUSTOM_STAT(FileIO, PakPrecacherHotBlocksCount, (int32)GPreCacheHotBlocksCount, ECsvCustomStatOp::Set);
 		CSV_CUSTOM_STAT(FileIO, PakPrecacherColdBlocksCount, (int32)GPreCacheColdBlocksCount, ECsvCustomStatOp::Set);
 		CSV_CUSTOM_STAT(FileIO, PakPrecacherTotalLoadedMB, (int32)(GPreCacheTotalLoaded/(1024*1024)), ECsvCustomStatOp::Set);
+		CSV_CUSTOM_STAT(FileIO, PakPrecacherBlockMemoryMB, (int32)(FPakPrecacher::Get().GetBlockMemory() / (1024 * 1024)), ECsvCustomStatOp::Set);
+
 
 		if (GPreCacheTotalLoadedLastTick != 0)
 		{
@@ -4193,6 +4281,12 @@ void FPakPlatformFile::Tick()
 			CSV_CUSTOM_STAT(FileIO, PakPrecacherPerFrameKB, (int32)diff, ECsvCustomStatOp::Set);
 		}
 		GPreCacheTotalLoadedLastTick = GPreCacheTotalLoaded;
+
+		CSV_CUSTOM_STAT(FileIO, PakPrecacherSeeks, (int32)GPreCacheSeeks, ECsvCustomStatOp::Set);
+		CSV_CUSTOM_STAT(FileIO, PakPrecacherBadSeeks, (int32)GPreCacheBadSeeks, ECsvCustomStatOp::Set);
+		CSV_CUSTOM_STAT(FileIO, PakPrecacherContiguousReads, (int32)GPreCacheContiguousReads, ECsvCustomStatOp::Set);
+		
+		CSV_CUSTOM_STAT(FileIO, PakLoads, (int32)PakPrecacherSingleton->Get().GetLoads(), ECsvCustomStatOp::Set);
 }
 #endif
 #if TRACK_DISK_UTILIZATION && CSV_PROFILER
