Build break fix

diff --git a/Engine/Source/Runtime/PakFile/Private/IPlatformFilePak.cpp b/Engine/Source/Runtime/PakFile/Private/IPlatformFilePak.cpp
--- a/Engine/Source/Runtime/PakFile/Private/IPlatformFilePak.cpp
+++ b/Engine/Source/Runtime/PakFile/Private/IPlatformFilePak.cpp
@@ -2424,6 +2424,7 @@ private: // below here we assume CachedFilesScopeLock until we get to the next s
 		RequestsToLower[IndexToFill].RequestHandle = Pak.Handle->ReadRequest(GetRequestOffset(Block.OffsetAndPakIndex), Block.Size, Priority, &CallbackFromLower);
 		RedundantReadTracker.CheckBlock(GetRequestOffset(Block.OffsetAndPakIndex), Block.Size);
 
+#if CSV_PROFILER
 		FJoinedOffsetAndPakIndex OldLastReadRequest = LastReadRequest;
 		LastReadRequest = Block.OffsetAndPakIndex + Block.Size;
 
@@ -2442,7 +2443,7 @@ private: // below here we assume CachedFilesScopeLock until we get to the next s
 		{
 			GPreCacheContiguousReads++;
 		}
-
+#endif
 		Loads++;
 		LoadSize += Block.Size;
 	}
