Fix crash when FChunkCacheWorker::CheckSignature and FPakPrecacher::DoSignatureCheck fail at the same time in different threads

diff --git a/Engine/Source/Runtime/PakFile/Private/IPlatformFilePak.cpp b/Engine/Source/Runtime/PakFile/Private/IPlatformFilePak.cpp
--- a/Engine/Source/Runtime/PakFile/Private/IPlatformFilePak.cpp
+++ b/Engine/Source/Runtime/PakFile/Private/IPlatformFilePak.cpp
@@ -180,11 +180,13 @@ FPakCustomEncryptionDelegate& FPakPlatformFile::GetPakCustomEncryptionDelegate()
 	return Delegate;
 }
 
+// Needs to be a global as multiple threads may try to access the failed delegate
+static FPakChunkSignatureCheckFailedHandler SignatureFailedDelegate;
 FPakChunkSignatureCheckFailedHandler& FPakPlatformFile::GetPakChunkSignatureCheckFailedHandler()
 {
-	static FPakChunkSignatureCheckFailedHandler Delegate;
-	return Delegate;
+	return SignatureFailedDelegate;
 }
+
 FPakMasterSignatureTableCheckFailureHandler& FPakPlatformFile::GetPakMasterSignatureTableCheckFailureHandler()
 {
 	static FPakMasterSignatureTableCheckFailureHandler Delegate;
