Enable time based pak cache trimming
Allow a read request to tell the pak cache not to keep the memory after the request is completed.
Make virtual textures give up their pak cache memory immediately after loading.

diff --git a/Engine/Source/Runtime/Core/Public/GenericPlatform/GenericPlatformFile.h b/Engine/Source/Runtime/Core/Public/GenericPlatform/GenericPlatformFile.h
--- a/Engine/Source/Runtime/Core/Public/GenericPlatform/GenericPlatformFile.h
+++ b/Engine/Source/Runtime/Core/Public/GenericPlatform/GenericPlatformFile.h
@@ -24,7 +24,8 @@ enum EAsyncIOPriorityAndFlags
 	AIOP_PRIORITY_MASK = 0x000000ff,
 
 	// Flags - combine with priorities if needed
-	AIOP_FLAG_PRECACHE = 0x00000100,
+	AIOP_FLAG_PRECACHE	=	0x00000100,
+	AIOP_FLAG_DONTCACHE	=	0x00000200,
 
 	// Priorities
 	AIOP_MIN = 0,
diff --git a/Engine/Source/Runtime/Engine/Private/Streaming/Texture2DStreamIn_IO.cpp b/Engine/Source/Runtime/Engine/Private/Streaming/Texture2DStreamIn_IO.cpp
--- a/Engine/Source/Runtime/Engine/Private/Streaming/Texture2DStreamIn_IO.cpp
+++ b/Engine/Source/Runtime/Engine/Private/Streaming/Texture2DStreamIn_IO.cpp
@@ -1,4 +1,4 @@
-// Copyright Epic Games, Inc. All Rights Reserved.
+// Copyright 1998-2019 Epic Games, Inc. All Rights Reserved.
 
 /*=============================================================================
 Texture2DStreamIn.cpp: Stream in helper for 2D textures using texture streaming files.
@@ -17,15 +17,19 @@ Texture2DStreamIn.cpp: Stream in helper for 2D textures using texture streaming
 FTexture2DStreamIn_IO::FTexture2DStreamIn_IO(UTexture2D* InTexture, int32 InRequestedMips, bool InPrioritizedIORequest)
 	: FTexture2DStreamIn(InTexture, InRequestedMips)
 	, bPrioritizedIORequest(InPrioritizedIORequest)
-
+	, IOFileOffset(0)
+	, IOFileHandle(nullptr)
 {
-	IORequests.AddZeroed(InTexture->GetNumMips());
+	FMemory::Memzero(IORequests, sizeof(IORequests));
 }
 
 FTexture2DStreamIn_IO::~FTexture2DStreamIn_IO()
 {
+	// Work must be done here because derived destructors have been called now and so derived members are invalid.
+	check(!IOFileHandle);
+
 #if DO_CHECK
-	for (const IBulkDataIORequest* IORequest : IORequests)
+	for (IAsyncReadRequest* IORequest : IORequests)
 	{
 		check(!IORequest);
 	}
@@ -37,49 +41,97 @@ void FTexture2DStreamIn_IO::SetIOFilename(const FContext& Context)
 	const TIndirectArray<FTexture2DMipMap>& OwnerMips = Context.Texture->GetPlatformMips();
 
 	const int32 CurrentFirstMip = Context.Resource->GetCurrentFirstMip();
-
-	if (PendingFirstMip < CurrentFirstMip && GEventDrivenLoaderEnabled)
+	for (int32 MipIndex = PendingFirstMip; MipIndex < CurrentFirstMip; ++MipIndex)
 	{
-		const FTexture2DMipMap& MipMap = OwnerMips[PendingFirstMip];
-	
-#if TEXTURE2DMIPMAP_USE_COMPACT_BULKDATA
-		verify(Context.Texture->GetMipDataFilename(PendingFirstMip, IOFilename));
+		const FTexture2DMipMap& MipMap = OwnerMips[MipIndex];
+		if (MipMap.BulkData.IsStoredCompressedOnDisk())
+		{
+			UE_LOG(LogTexture, Error, TEXT("Compression at the package level is no longer supported."));
+			IOFilename.Reset();
+			break;
+		}
+		else if (MipMap.BulkData.GetBulkDataSize() <= 0)
+		{
+			UE_LOG(LogTexture, Error, TEXT("%s has invalid bulk data size."), *Context.Texture->GetName());
+			IOFilename.Reset();
+			break;
+		}
+
+		if (MipIndex == PendingFirstMip)
+		{
+#if !TEXTURE2DMIPMAP_USE_COMPACT_BULKDATA
+			IOFilename = MipMap.BulkData.GetFilename();
+#else
+			verify(Context.Texture->GetMipDataFilename(MipIndex, IOFilename));
+#endif
+
+			if (GEventDrivenLoaderEnabled)
+			{
+				if (IOFilename.EndsWith(TEXT(".uasset")) || IOFilename.EndsWith(TEXT(".umap")))
+				{
+					IOFileOffset = -IFileManager::Get().FileSize(*IOFilename);
+					check(IOFileOffset < 0);
+					IOFilename = FPaths::GetBaseFilename(IOFilename, false) + TEXT(".uexp");
+					UE_LOG(LogTexture, Error, TEXT("Streaming from the .uexp file '%s' this MUST be in a ubulk instead for best performance."), *IOFilename);
+				}
+			}
+		}
+#if !TEXTURE2DMIPMAP_USE_COMPACT_BULKDATA
+		else if (IOFilename != MipMap.BulkData.GetFilename())
+		{
+			UE_LOG(LogTexture, Error, TEXT("All of the streaming mips must be stored in the same file %s %s."), *IOFilename, *MipMap.BulkData.GetFilename());
+			IOFilename.Reset();
+			break;
+		}
 #endif
 	}
+
+	if (IOFilename.IsEmpty())
+	{
+		MarkAsCancelled();
+	}
 }
 
 void FTexture2DStreamIn_IO::SetIORequests(const FContext& Context)
 {
 	SetAsyncFileCallback();
-	
-	const TIndirectArray<FTexture2DMipMap>& OwnerMips = Context.Texture->GetPlatformMips();
-	const int32 CurrentFirstMip = Context.Resource->GetCurrentFirstMip();
 
-	for (int32 MipIndex = PendingFirstMip; MipIndex < CurrentFirstMip && !IsCancelled(); ++MipIndex)
+	check(!IOFileHandle);
+	IOFileHandle = FPlatformFileManager::Get().GetPlatformFile().OpenAsyncRead(*IOFilename);
+	if (IOFileHandle)
 	{
-		const FTexture2DMipMap& MipMap = OwnerMips[MipIndex];
+		const TIndirectArray<FTexture2DMipMap>& OwnerMips = Context.Texture->GetPlatformMips();
+		const int32 CurrentFirstMip = Context.Resource->GetCurrentFirstMip();
+
+		for (int32 MipIndex = PendingFirstMip; MipIndex < CurrentFirstMip && !IsCancelled(); ++MipIndex)
+		{
+			const FTexture2DMipMap& MipMap = OwnerMips[MipIndex];
 
-		check(MipData[MipIndex]);
+			check(MipData[MipIndex]);
 
-		// Increment as we push the requests. If a requests complete immediately, then it will call the callback
-		// but that won't do anything because the tick would not try to acquire the lock since it is already locked.
-		TaskSynchronization.Increment();
+			// Increment as we push the requests. If a requests complete immediately, then it will call the callback
+			// but that won't do anything because the tick would not try to acquire the lock since it is already locked.
+			TaskSynchronization.Increment();
 
-		IORequests[MipIndex] = MipMap.BulkData.CreateStreamingRequest(
-			TEXTURE2DMIPMAP_PARAM(IOFilename) // Only used if TEXTURE2DMIPMAP_USE_COMPACT_BULKDATA is enabled
-			0,
-			MipMap.BulkData.GetBulkDataSize(),
-			bPrioritizedIORequest ? AIOP_BelowNormal : AIOP_Low,
-			&AsyncFileCallBack,
-			(uint8*)MipData[MipIndex]);
+			IORequests[MipIndex] = IOFileHandle->ReadRequest(
+				MipMap.BulkData.GetBulkDataOffsetInFile() + IOFileOffset, 
+				MipMap.BulkData.GetBulkDataSize(), 
+				bPrioritizedIORequest ? (AIOP_FLAG_DONTCACHE|AIOP_BelowNormal) : (AIOP_FLAG_DONTCACHE|AIOP_Low), 
+				&AsyncFileCallBack, 
+				(uint8*)MipData[MipIndex]);
+		}
+	}
+	else
+	{
+		MarkAsCancelled();
 	}
 }
 
 void FTexture2DStreamIn_IO::CancelIORequests()
 {
-	for (int32 MipIndex = 0; MipIndex < IORequests.Num(); ++MipIndex)
+	for (int32 MipIndex = 0; MipIndex < MAX_TEXTURE_MIP_COUNT; ++MipIndex)
 	{
-		IBulkDataIORequest* IORequest = IORequests[MipIndex];
+		IAsyncReadRequest* IORequest = IORequests[MipIndex];
 		if (IORequest)
 		{
 			// Calling cancel will trigger the SetAsyncFileCallback() which will also try a tick but will fail.
@@ -90,24 +142,30 @@ void FTexture2DStreamIn_IO::CancelIORequests()
 
 void FTexture2DStreamIn_IO::ClearIORequests(const FContext& Context)
 {
-	const TIndirectArray<FTexture2DMipMap>& OwnerMips = Context.Texture->GetPlatformMips();
-	const int32 CurrentFirstMip = Context.Resource->GetCurrentFirstMip();
-
-	for (int32 MipIndex = PendingFirstMip; MipIndex < CurrentFirstMip; ++MipIndex)
+	if (IOFileHandle)
 	{
-		IBulkDataIORequest* IORequest = IORequests[MipIndex];
-		IORequests[MipIndex] = nullptr;
+		const TIndirectArray<FTexture2DMipMap>& OwnerMips = Context.Texture->GetPlatformMips();
+		const int32 CurrentFirstMip = Context.Resource->GetCurrentFirstMip();
 
-		if (IORequest != nullptr)
+		for (int32 MipIndex = PendingFirstMip; MipIndex < CurrentFirstMip; ++MipIndex)
 		{
-			// If clearing requests not yet completed, cancel and wait.
-			if (!IORequest->PollCompletion())
+			IAsyncReadRequest* IORequest = IORequests[MipIndex];
+			IORequests[MipIndex] = nullptr;
+
+			if (IORequest)
 			{
-				IORequest->Cancel();
-				IORequest->WaitCompletion();
+				// If clearing requests not yet completed, cancel and wait.
+				if (!IORequest->PollCompletion())
+				{
+					IORequest->Cancel();
+					IORequest->WaitCompletion();
+				}
+				delete IORequest;
 			}
-			delete IORequest;
 		}
+
+		delete IOFileHandle;
+		IOFileHandle = nullptr;
 	}
 }
 
@@ -143,28 +201,16 @@ void FTexture2DStreamIn_IO::Abort()
 	{
 		FTexture2DStreamIn::Abort();
 
-		if (HasPendingIORequests())
+		// IO requests can only exist in the lifetime of IOFileHandle.
+		if (IOFileHandle)
 		{
 			// Prevent the update from being considered done before this is finished.
-			// By checking that it was not already canceled, we make sure this doesn't get called twice.
+			// By checking that it was not already cancelled, we make sure this doesn't get called twice.
 			(new FAsyncCancelIORequestsTask(this))->StartBackgroundTask();
 		}
 	}
 }
 
-bool FTexture2DStreamIn_IO::HasPendingIORequests()
-{
-	for (IBulkDataIORequest* IORequest : IORequests)
-	{
-		if (IORequest != nullptr)
-		{
-			return true;
-		}
-	}
-
-	return false;
-}
-
 void FTexture2DStreamIn_IO::FCancelIORequestsTask::DoWork()
 {
 	check(PendingUpdate);
diff --git a/Engine/Source/Runtime/Engine/Private/VT/UploadingVirtualTexture.cpp b/Engine/Source/Runtime/Engine/Private/VT/UploadingVirtualTexture.cpp
--- a/Engine/Source/Runtime/Engine/Private/VT/UploadingVirtualTexture.cpp
+++ b/Engine/Source/Runtime/Engine/Private/VT/UploadingVirtualTexture.cpp
@@ -47,9 +47,9 @@ static EAsyncIOPriorityAndFlags GetAsyncIOPriority(EVTRequestPagePriority Priori
 {
 	switch (Priority)
 	{
-	case EVTRequestPagePriority::High: return (EAsyncIOPriorityAndFlags)GVirtualTextureIOPriority_HighPagePri;
-	case EVTRequestPagePriority::Normal: return (EAsyncIOPriorityAndFlags)GVirtualTextureIOPriority_NormalPagePri;
-	default: check(false); return (EAsyncIOPriorityAndFlags)GVirtualTextureIOPriority_NormalPagePri;
+	case EVTRequestPagePriority::High: return (EAsyncIOPriorityAndFlags)GVirtualTextureIOPriority_HighPagePri | EAsyncIOPriorityAndFlags::AIOP_FLAG_DONTCACHE;
+	case EVTRequestPagePriority::Normal: return (EAsyncIOPriorityAndFlags)GVirtualTextureIOPriority_NormalPagePri | EAsyncIOPriorityAndFlags::AIOP_FLAG_DONTCACHE;
+	default: check(false); return (EAsyncIOPriorityAndFlags)GVirtualTextureIOPriority_NormalPagePri | EAsyncIOPriorityAndFlags::AIOP_FLAG_DONTCACHE;
 	}
 }
 
diff --git a/Engine/Source/Runtime/PakFile/Private/IPlatformFilePak.cpp b/Engine/Source/Runtime/PakFile/Private/IPlatformFilePak.cpp
--- a/Engine/Source/Runtime/PakFile/Private/IPlatformFilePak.cpp
+++ b/Engine/Source/Runtime/PakFile/Private/IPlatformFilePak.cpp
@@ -1,4 +1,4 @@
-// Copyright Epic Games, Inc. All Rights Reserved.
+// Copyright 1998-2019 Epic Games, Inc. All Rights Reserved.
 
 #include "IPlatformFilePak.h"
 #include "HAL/FileManager.h"
@@ -33,6 +33,8 @@
 
 #include "Async/MappedFileHandle.h"
 
+PRAGMA_DISABLE_OPTIMIZATION_ACTUAL
+
 DEFINE_LOG_CATEGORY(LogPakFile);
 
 DEFINE_STAT(STAT_PakFile_Read);
@@ -56,7 +58,7 @@ CSV_DECLARE_CATEGORY_MODULE_EXTERN(CORE_API, FileIO);
 
 static FString GMountStartupPaksWildCard = TEXT(MOUNT_STARTUP_PAKS_WILDCARD);
 
-int32 GetPakchunkIDFromPakFile(const FString& InFilename)
+int32 ParseChunkIDFromFilename(const FString& InFilename)
 {
 	FString ChunkIdentifier(TEXT("pakchunk"));
 	FString BaseFilename = FPaths::GetBaseFilename(InFilename);
@@ -180,13 +182,11 @@ FPakCustomEncryptionDelegate& FPakPlatformFile::GetPakCustomEncryptionDelegate()
 	return Delegate;
 }
 
-// Needs to be a global as multiple threads may try to access the failed delegate
-static FPakChunkSignatureCheckFailedHandler SignatureFailedDelegate;
 FPakChunkSignatureCheckFailedHandler& FPakPlatformFile::GetPakChunkSignatureCheckFailedHandler()
 {
-	return SignatureFailedDelegate;
+	static FPakChunkSignatureCheckFailedHandler Delegate;
+	return Delegate;
 }
-
 FPakMasterSignatureTableCheckFailureHandler& FPakPlatformFile::GetPakMasterSignatureTableCheckFailureHandler()
 {
 	static FPakMasterSignatureTableCheckFailureHandler Delegate;
@@ -223,11 +223,7 @@ void FPakPlatformFile::GetFilenamesInPakFile(const FString& InPakFilename, TArra
 	}
 }
 
-#if !defined(PLATFORM_BYPASS_PAK_PRECACHE)
-	#error "PLATFORM_BYPASS_PAK_PRECACHE must be defined."
-#endif
-
-#define USE_PAK_PRECACHE (!PLATFORM_BYPASS_PAK_PRECACHE && !IS_PROGRAM && !WITH_EDITOR) // you can turn this off to use the async IO stuff without the precache
+#define USE_PAK_PRECACHE (!IS_PROGRAM && !WITH_EDITOR) // you can turn this off to use the async IO stuff without the precache
 
 /**
 * Precaching
@@ -348,23 +344,6 @@ static FAutoConsoleVariableRef CVar_ForceDecompressionFails(
 );
 #endif
 
-class FPakSizeRequest : public IAsyncReadRequest
-{
-public:
-	FPakSizeRequest(FAsyncFileCallBack* CompleteCallback, int64 InFileSize)
-		: IAsyncReadRequest(CompleteCallback, true, nullptr)
-	{
-		Size = InFileSize;
-		SetComplete();
-	}
-	virtual void WaitCompletionImpl(float TimeLimitSeconds) override
-	{
-	}
-	virtual void CancelImpl()
-	{
-	}
-};
-
 #if USE_PAK_PRECACHE
 #include "Async/TaskGraphInterfaces.h"
 #define PAK_CACHE_GRANULARITY (1024*64)
@@ -381,10 +360,6 @@ volatile int64 GPreCacheHotBlocksCount = 0;
 volatile int64 GPreCacheColdBlocksCount = 0;
 volatile int64 GPreCacheTotalLoaded = 0;
 int64 GPreCacheTotalLoadedLastTick = 0;
-
-volatile int64 GPreCacheSeeks = 0;
-volatile int64 GPreCacheBadSeeks = 0;
-volatile int64 GPreCacheContiguousReads = 0;
 #endif
 
 DECLARE_FLOAT_ACCUMULATOR_STAT(TEXT("PakCache Signing Chunk Hash Time"), STAT_PakCache_SigningChunkHashTime, STATGROUP_PakFile);
@@ -448,6 +423,13 @@ static FAutoConsoleVariableRef CVar_PakCache_TimeToTrim(
 	TEXT("Controls how long to hold onto a cached but unreferenced block for.")
 );
 
+int32 GPakCache_EnableNoCaching = 0;
+static FAutoConsoleVariableRef CVar_EnableNoCaching(
+	TEXT("pakcache.EnableNoCaching"),
+	GPakCache_EnableNoCaching,
+	TEXT("if > 0, then we'll allow a read requests pak cache memory to be ditched early")
+);
+
 class FPakPrecacher;
 
 typedef uint64 FJoinedOffsetAndPakIndex;
@@ -1148,7 +1130,6 @@ class FPakPrecacher
 		TIntervalTreeIndex Index;
 		TIntervalTreeIndex Next;
 		EBlockStatus Status;
-		double TimeNoLongerReferenced;
 
 		FCacheBlock()
 			: OffsetAndPakIndex(0)
@@ -1158,6 +1139,7 @@ class FPakPrecacher
 			, Index(IntervalTreeInvalidIndex)
 			, Next(IntervalTreeInvalidIndex)
 			, Status(EBlockStatus::InFlight)
+			, TimeNoLongerReferenced(0)
 		{
 		}
 	};
@@ -1300,8 +1282,6 @@ class FPakPrecacher
 	FCriticalSection SetAsyncMinimumPriorityScopeLock;
 	bool bEnableSignatureChecks;
 public:
-	int64 GetBlockMemory() { return BlockMemory; }
-	int64 GetBlockMemoryHighWater() { return BlockMemoryHighWater; }
 
 	static void Init(IPlatformFile* InLowerLevel, bool bInEnableSignatureChecks) 
 	{
@@ -1357,7 +1337,7 @@ public:
 		, bEnableSignatureChecks(bInEnableSignatureChecks)
 	{
 		check(LowerLevel && FPlatformProcess::SupportsMultithreading());
-//		GPakCache_MaxRequestsToLowerLevel = FMath::Max(FMath::Min(FPlatformMisc::NumberOfIOWorkerThreadsToSpawn(), GPakCache_MaxRequestsToLowerLevel), 1);
+		GPakCache_MaxRequestsToLowerLevel = FMath::Max(FMath::Min(FPlatformMisc::NumberOfIOWorkerThreadsToSpawn(), GPakCache_MaxRequestsToLowerLevel), 1);
 		check(GPakCache_MaxRequestsToLowerLevel <= PAK_CACHE_MAX_REQUESTS);
 	}
 
@@ -1552,9 +1532,10 @@ private: // below here we assume CachedFilesScopeLock until we get to the next s
 		return MAX_uint64;
 	}
 
-	bool AddRequest(FPakInRequest& Request, TIntervalTreeIndex NewIndex)
+	bool AddRequest(TIntervalTreeIndex NewIndex)
 	{
 		// CachedFilesScopeLock is locked
+		FPakInRequest& Request = InRequestAllocator.Get(NewIndex);
 		uint16 PakIndex = GetRequestPakIndex(Request.OffsetAndPakIndex);
 		int64 Offset = GetRequestOffset(Request.OffsetAndPakIndex);
 		FPakData& Pak = CachedPakData[PakIndex];
@@ -1717,7 +1698,7 @@ private: // below here we assume CachedFilesScopeLock until we get to the next s
 	void TrimCache(bool bDiscardAll = false, uint16 StartPakIndex = 65535)
 	{
 
-		if (GPakCache_UseNewTrim && !bDiscardAll)
+		if (GPakCache_UseNewTrim)
 		{
 			StartPakIndex = 0;
 			uint16 EndPakIndex = CachedPakData.Num();
@@ -1727,7 +1708,7 @@ private: // below here we assume CachedFilesScopeLock until we get to the next s
 			CacheVisitedAlready.AddDefaulted(OffsetAndPakIndexOfSavedBlocked.Num());
 
 			int64 MemoryBudget = GPakCache_MaxBlockMemory * (1024 * 1024);
-			bool AlreadyRemovedBlocksBecauseOfMemoryOverage = false;
+
 
 			while (BlockMemory > MemoryBudget)
 			{
@@ -1782,7 +1763,6 @@ private: // below here we assume CachedFilesScopeLock until we get to the next s
 							);
 						}
 						OffsetAndPakIndexOfSavedBlocked[CacheIndex].RemoveAt(0, NumToRemove, false);
-						AlreadyRemovedBlocksBecauseOfMemoryOverage = true;
 					}
 				}
 				if (NoneToRemove)
@@ -1790,7 +1770,7 @@ private: // below here we assume CachedFilesScopeLock until we get to the next s
 					break;
 				}
 			}
-			if(GPakCache_TimeToTrim != 0.0f)
+			if (GPakCache_TimeToTrim != 0.0f)
 			{
 				// we'll trim based on time rather than trying to keep within a memory budget
 				double CurrentTime = FPlatformTime::Seconds();
@@ -1809,7 +1789,7 @@ private: // below here we assume CachedFilesScopeLock until we get to the next s
 							uint16 PakIndex = GetRequestPakIndex(OffsetAndPakIndex);
 							int64 Offset = GetRequestOffset(OffsetAndPakIndex);
 							FPakData& Pak = CachedPakData[PakIndex];
-							bool removedall = true;
+							bool bRemovedAll = true;
 							MaybeRemoveOverlappingNodesInIntervalTree<FCacheBlock>(
 								&Pak.CacheBlocks[(int32)EBlockStatus::Complete],
 								CacheBlockAllocator,
@@ -1819,7 +1799,7 @@ private: // below here we assume CachedFilesScopeLock until we get to the next s
 								Pak.MaxNode,
 								Pak.StartShift,
 								Pak.MaxShift,
-								[this, CurrentTime, &removedall](TIntervalTreeIndex BlockIndex) -> bool
+								[this, CurrentTime, &bRemovedAll](TIntervalTreeIndex BlockIndex) -> bool
 							{
 								FCacheBlock &Block = CacheBlockAllocator.Get(BlockIndex);
 								if (!Block.InRequestRefCount && (CurrentTime - Block.TimeNoLongerReferenced >= GPakCache_TimeToTrim))
@@ -1828,11 +1808,11 @@ private: // below here we assume CachedFilesScopeLock until we get to the next s
 									ClearBlock(Block);
 									return true;
 								}
-								removedall = false;
+								bRemovedAll = false;
 								return false;
 							}
 							);
-							if (!removedall)
+							if (!bRemovedAll)
 								break;
 							NumToRemove++;
 						}
@@ -1909,6 +1889,9 @@ private: // below here we assume CachedFilesScopeLock until we get to the next s
 		FPakData& Pak = CachedPakData[PakIndex];
 		check(Offset + Request.Size <= Pak.TotalSize && Request.Size > 0 && Request.GetPriority() >= AIOP_MIN && Request.GetPriority() <= AIOP_MAX && int32(Request.Status) >= 0 && int32(Request.Status) < int32(EInRequestStatus::Num));
 
+		bool RequestDontCache = (Request.PriorityAndFlags & AIOP_FLAG_DONTCACHE) != 0;
+
+
 		if (RemoveFromIntervalTree<FPakInRequest>(&Pak.InRequests[Request.GetPriority()][(int32)Request.Status], InRequestAllocator, Index, Pak.StartShift, Pak.MaxShift))
 		{
 
@@ -1922,7 +1905,7 @@ private: // below here we assume CachedFilesScopeLock until we get to the next s
 				Pak.MaxNode,
 				Pak.StartShift,
 				Pak.MaxShift,
-				[this, OffsetOfLastByte](TIntervalTreeIndex BlockIndex) -> bool
+				[this, OffsetOfLastByte, RequestDontCache](TIntervalTreeIndex BlockIndex) -> bool
 			{
 				FCacheBlock &Block = CacheBlockAllocator.Get(BlockIndex);
 				check(Block.InRequestRefCount);
@@ -1930,11 +1913,23 @@ private: // below here we assume CachedFilesScopeLock until we get to the next s
 				{
 					if (GPakCache_NumUnreferencedBlocksToCache && GetRequestOffset(Block.OffsetAndPakIndex) + Block.Size > OffsetOfLastByte) // last block
 					{
-						uint16 BlocksPakIndex = GetRequestPakIndexLow(Block.OffsetAndPakIndex);
-						int32 BlocksCacheIndex = CachedPakData[BlocksPakIndex].ActualPakFile->GetCacheIndex();
-						Block.TimeNoLongerReferenced = FPlatformTime::Seconds();
-						OffsetAndPakIndexOfSavedBlocked[BlocksCacheIndex].Remove(Block.OffsetAndPakIndex);
-						OffsetAndPakIndexOfSavedBlocked[BlocksCacheIndex].Add(Block.OffsetAndPakIndex);
+						if (RequestDontCache && GPakCache_EnableNoCaching != 0)
+						{
+							uint16 BlocksPakIndex = GetRequestPakIndexLow(Block.OffsetAndPakIndex);
+							int32 BlocksCacheIndex = CachedPakData[BlocksPakIndex].ActualPakFile->GetCacheIndex();
+							Block.TimeNoLongerReferenced = 0.0;
+							OffsetAndPakIndexOfSavedBlocked[BlocksCacheIndex].Remove(Block.OffsetAndPakIndex);
+							ClearBlock(Block);
+							return true;
+						}
+						else
+						{
+							uint16 BlocksPakIndex = GetRequestPakIndexLow(Block.OffsetAndPakIndex);
+							int32 BlocksCacheIndex = CachedPakData[BlocksPakIndex].ActualPakFile->GetCacheIndex();
+
+							OffsetAndPakIndexOfSavedBlocked[BlocksCacheIndex].Remove(Block.OffsetAndPakIndex);
+							OffsetAndPakIndexOfSavedBlocked[BlocksCacheIndex].Add(Block.OffsetAndPakIndex);
+						}
 						return false;
 					}
 					ClearBlock(Block);
@@ -2425,27 +2420,7 @@ private: // below here we assume CachedFilesScopeLock until we get to the next s
 
 		RequestsToLower[IndexToFill].RequestHandle = Pak.Handle->ReadRequest(GetRequestOffset(Block.OffsetAndPakIndex), Block.Size, Priority, &CallbackFromLower);
 		RedundantReadTracker.CheckBlock(GetRequestOffset(Block.OffsetAndPakIndex), Block.Size);
-
-#if CSV_PROFILER
-		FJoinedOffsetAndPakIndex OldLastReadRequest = LastReadRequest;
 		LastReadRequest = Block.OffsetAndPakIndex + Block.Size;
-
-		if (OldLastReadRequest != Block.OffsetAndPakIndex)
-		{
-			if (GetRequestPakIndexLow(OldLastReadRequest) != GetRequestPakIndexLow(Block.OffsetAndPakIndex))
-			{
-				GPreCacheBadSeeks++;
-			}
-			else
-			{
-				GPreCacheSeeks++;
-			}
-		}
-		else
-		{
-			GPreCacheContiguousReads++;
-		}
-#endif
 		Loads++;
 		LoadSize += Block.Size;
 	}
@@ -2692,7 +2667,7 @@ public:
 		OutstandingRequests.Add(Request.UniqueID, RequestIndex);
 		RequestCounter.Increment();
 
-		if (AddRequest(Request, RequestIndex))
+		if (AddRequest(RequestIndex))
 		{
 #if USE_PAK_PRECACHE && CSV_PROFILER
 			FPlatformAtomics::InterlockedIncrement(&GPreCacheHotBlocksCount);
@@ -3275,6 +3250,23 @@ public:
 	}
 };
 
+class FPakSizeRequest : public IAsyncReadRequest
+{
+public:
+	FPakSizeRequest(FAsyncFileCallBack* CompleteCallback, int64 InFileSize)
+		: IAsyncReadRequest(CompleteCallback, true, nullptr)
+	{
+		Size = InFileSize;
+		SetComplete();
+	}
+	virtual void WaitCompletionImpl(float TimeLimitSeconds) override
+	{
+	}
+	virtual void CancelImpl()
+	{
+	}
+};
+
 class FPakProcessedReadRequest : public IAsyncReadRequest
 {
 	FPakAsyncReadFileHandle* Owner;
@@ -4158,65 +4150,6 @@ void FPakPlatformFile::TrackPak(const TCHAR* Filename, const FPakEntry* PakEntry
 }
 #endif
 
-class FBypassPakAsyncReadFileHandle final : public IAsyncReadFileHandle
-{
-	FName PakFile;
-	int64 PakFileSize;
-	int64 OffsetInPak;
-	int64 UncompressedFileSize;
-	FPakEntry FileEntry;
-	IAsyncReadFileHandle* LowerHandle;
-
-public:
-	FBypassPakAsyncReadFileHandle(const FPakEntry* InFileEntry, FPakFile* InPakFile, const TCHAR* Filename)
-		: PakFile(InPakFile->GetFilenameName())
-		, PakFileSize(InPakFile->TotalSize())
-		, FileEntry(*InFileEntry)
-	{
-		OffsetInPak = FileEntry.Offset + FileEntry.GetSerializedSize(InPakFile->GetInfo().Version);
-		UncompressedFileSize = FileEntry.UncompressedSize;
-		int64 CompressedFileSize = FileEntry.UncompressedSize;
-		check(FileEntry.CompressionMethodIndex == 0);
-		UE_LOG(LogPakFile, Verbose, TEXT("FPakPlatformFile::OpenAsyncRead (FBypassPakAsyncReadFileHandle)[%016llX, %016llX) %s"), OffsetInPak, OffsetInPak + CompressedFileSize, Filename);
-		check(PakFileSize > 0 && OffsetInPak + CompressedFileSize <= PakFileSize && OffsetInPak >= 0);
-
-		LowerHandle = IPlatformFile::GetPlatformPhysical().OpenAsyncRead(*InPakFile->GetFilename());
-	}
-	~FBypassPakAsyncReadFileHandle()
-	{
-		delete LowerHandle;
-	}
-
-	virtual IAsyncReadRequest* SizeRequest(FAsyncFileCallBack* CompleteCallback = nullptr) override
-	{
-		if (!LowerHandle)
-		{
-			return nullptr;
-		}
-		return new FPakSizeRequest(CompleteCallback, UncompressedFileSize);
-	}
-	virtual IAsyncReadRequest* ReadRequest(int64 Offset, int64 BytesToRead, EAsyncIOPriorityAndFlags PriorityAndFlags = AIOP_Normal, FAsyncFileCallBack* CompleteCallback = nullptr, uint8* UserSuppliedMemory = nullptr) override
-	{
-		if (!LowerHandle)
-		{
-			return nullptr;
-		}
-		if (BytesToRead == MAX_int64)
-		{
-			BytesToRead = UncompressedFileSize - Offset;
-		}
-		check(Offset + BytesToRead <= UncompressedFileSize && Offset >= 0);
-		check(FileEntry.CompressionMethodIndex == 0);
-		check(Offset + BytesToRead + OffsetInPak <= PakFileSize);
-
-		return LowerHandle->ReadRequest(Offset + OffsetInPak, BytesToRead, PriorityAndFlags, CompleteCallback, UserSuppliedMemory);
-	}
-	virtual bool UsesCache() override
-	{
-		return LowerHandle->UsesCache();
-	}
-};
-
 IAsyncReadFileHandle* FPakPlatformFile::OpenAsyncRead(const TCHAR* Filename)
 {
 	CSV_SCOPED_TIMING_STAT(FileIO, PakOpenAsyncRead);
@@ -4236,19 +4169,6 @@ IAsyncReadFileHandle* FPakPlatformFile::OpenAsyncRead(const TCHAR* Filename)
 			return new FPakAsyncReadFileHandle(&FileEntry, PakFile, Filename);
 		}
 	}
-#elif PLATFORM_BYPASS_PAK_PRECACHE
-	{
-		FPakEntry FileEntry;
-		FPakFile* PakFile = NULL;
-		bool bFoundEntry = FindFileInPakFiles(Filename, &PakFile, &FileEntry);
-		if (bFoundEntry && PakFile && PakFile->GetFilenameName() != NAME_None && FileEntry.CompressionMethodIndex == 0 && !FileEntry.IsEncrypted())
-		{
-#if PAK_TRACKER
-			TrackPak(Filename, &FileEntry);
-#endif
-			return new FBypassPakAsyncReadFileHandle(&FileEntry, PakFile, Filename);
-		}
-	}
 #endif
 	return IPlatformFile::OpenAsyncRead(Filename);
 }
@@ -4260,8 +4180,6 @@ void FPakPlatformFile::SetAsyncMinimumPriority(EAsyncIOPriorityAndFlags Priority
 	{
 		FPakPrecacher::Get().SetAsyncMinimumPriority(Priority);
 	}
-#elif PLATFORM_BYPASS_PAK_PRECACHE
-	IPlatformFile::GetPlatformPhysical().SetAsyncMinimumPriority(Priority);
 #endif
 }
 
@@ -4274,8 +4192,6 @@ void FPakPlatformFile::Tick()
 		CSV_CUSTOM_STAT(FileIO, PakPrecacherHotBlocksCount, (int32)GPreCacheHotBlocksCount, ECsvCustomStatOp::Set);
 		CSV_CUSTOM_STAT(FileIO, PakPrecacherColdBlocksCount, (int32)GPreCacheColdBlocksCount, ECsvCustomStatOp::Set);
 		CSV_CUSTOM_STAT(FileIO, PakPrecacherTotalLoadedMB, (int32)(GPreCacheTotalLoaded/(1024*1024)), ECsvCustomStatOp::Set);
-		CSV_CUSTOM_STAT(FileIO, PakPrecacherBlockMemoryMB, (int32)(FPakPrecacher::Get().GetBlockMemory() / (1024 * 1024)), ECsvCustomStatOp::Set);
-
 
 		if (GPreCacheTotalLoadedLastTick != 0)
 		{
@@ -4284,12 +4200,6 @@ void FPakPlatformFile::Tick()
 			CSV_CUSTOM_STAT(FileIO, PakPrecacherPerFrameKB, (int32)diff, ECsvCustomStatOp::Set);
 		}
 		GPreCacheTotalLoadedLastTick = GPreCacheTotalLoaded;
-
-		CSV_CUSTOM_STAT(FileIO, PakPrecacherSeeks, (int32)GPreCacheSeeks, ECsvCustomStatOp::Set);
-		CSV_CUSTOM_STAT(FileIO, PakPrecacherBadSeeks, (int32)GPreCacheBadSeeks, ECsvCustomStatOp::Set);
-		CSV_CUSTOM_STAT(FileIO, PakPrecacherContiguousReads, (int32)GPreCacheContiguousReads, ECsvCustomStatOp::Set);
-		
-		CSV_CUSTOM_STAT(FileIO, PakLoads, (int32)PakPrecacherSingleton->Get().GetLoads(), ECsvCustomStatOp::Set);
 }
 #endif
 #if TRACK_DISK_UTILIZATION && CSV_PROFILER
@@ -4750,7 +4660,7 @@ FPakFile::FPakFile(const TCHAR* Filename, bool bIsSigned)
 	, bSigned(bIsSigned)
 	, bIsValid(false)
 	, bFilenamesRemoved(false)
-	, ChunkID(GetPakchunkIDFromPakFile(Filename))
+	, ChunkID(ParseChunkIDFromFilename(Filename))
 	, bAttemptedPakEntryShrink(false)
 	, bAttemptedPakFilenameUnload(false)
  	, MappedFileHandle(nullptr)
@@ -4780,7 +4690,7 @@ FPakFile::FPakFile(IPlatformFile* LowerLevel, const TCHAR* Filename, bool bIsSig
 	, bSigned(bIsSigned)
 	, bIsValid(false)
 	, bFilenamesRemoved(false)
-	, ChunkID(GetPakchunkIDFromPakFile(Filename))
+	, ChunkID(ParseChunkIDFromFilename(Filename))
 	, bAttemptedPakEntryShrink(false)
 	, bAttemptedPakFilenameUnload(false)
 	, MappedFileHandle(nullptr)
@@ -6050,14 +5960,11 @@ void FPakPlatformFile::FindPakFilesInDirectory(IPlatformFile* LowLevelFile, cons
 		TArray<FString>& FoundPakFiles;
 		IPlatformChunkInstall* ChunkInstall = nullptr;
 		FString WildCard;
-		bool bSkipOptionalPakFiles;
-
 	public:
-		FPakSearchVisitor(TArray<FString>& InFoundPakFiles, const FString& InWildCard, IPlatformChunkInstall* InChunkInstall, bool bInSkipOptionalPakFiles)
+		FPakSearchVisitor(TArray<FString>& InFoundPakFiles, const FString& InWildCard, IPlatformChunkInstall* InChunkInstall)
 			: FoundPakFiles(InFoundPakFiles)
 			, ChunkInstall(InChunkInstall)
 			, WildCard(InWildCard)
-			, bSkipOptionalPakFiles(bInSkipOptionalPakFiles)
 		{}
 		virtual bool Visit(const TCHAR* FilenameOrDirectory, bool bIsDirectory)
 		{
@@ -6069,32 +5976,23 @@ void FPakPlatformFile::FindPakFilesInDirectory(IPlatformFile* LowLevelFile, cons
 					// if a platform supports chunk style installs, make sure that the chunk a pak file resides in is actually fully installed before accepting pak files from it
 					if (ChunkInstall)
 					{
-						int32 PakchunkID = GetPakchunkIDFromPakFile(Filename);
-						if (PakchunkID != INDEX_NONE)
+						int32 ChunkID = ParseChunkIDFromFilename(Filename);
+						if (ChunkID != INDEX_NONE)
 						{
-							if (ChunkInstall->GetPakchunkLocation(PakchunkID) == EChunkLocation::NotAvailable)
+							if (ChunkInstall->GetChunkLocation(ChunkID) == EChunkLocation::NotAvailable)
 							{
 								return true;
 							}
 						}
 					}
-
-#if !UE_BUILD_SHIPPING
-					if (bSkipOptionalPakFiles == false || Filename.Find("optional") == INDEX_NONE)
-#endif
-					{
-						FoundPakFiles.Add(Filename);
-					}
+					FoundPakFiles.Add(Filename);
 				}
 			}
 			return true;
 		}
 	};
-
-	bool bSkipOptionalPakFiles = FParse::Param(FCommandLine::Get(), TEXT("SkipOptionalPakFiles"));
-
 	// Find all pak files.
-	FPakSearchVisitor Visitor(OutPakFiles, WildCard, FPlatformMisc::GetPlatformChunkInstall(), bSkipOptionalPakFiles);
+	FPakSearchVisitor Visitor(OutPakFiles, WildCard, FPlatformMisc::GetPlatformChunkInstall());
 	LowLevelFile->IterateDirectoryRecursively(Directory, Visitor);
 }
 
@@ -6143,14 +6041,12 @@ bool FPakPlatformFile::CheckIfPakFilesExist(IPlatformFile* LowLevelFile, const T
 bool FPakPlatformFile::ShouldBeUsed(IPlatformFile* Inner, const TCHAR* CmdLine) const
 {
 	bool Result = false;
-#if (!WITH_EDITOR || IS_MONOLITHIC)
 	if (!FParse::Param(CmdLine, TEXT("NoPak")))
 	{
 		TArray<FString> PakFolders;
 		GetPakFolders(CmdLine, PakFolders);
 		Result = CheckIfPakFilesExist(Inner, PakFolders);
 	}
-#endif
 	return Result;
 }
 
@@ -6576,7 +6472,7 @@ void FPakPlatformFile::RegisterEncryptionKey(const FGuid& InGuid, const FAES::FA
 				UE_LOG(LogPakFile, Log, TEXT("Successfully mounted deferred pak file '%s'"), *Entry.Filename);
 				NumMounted++;
 
-				int32 ChunkID = GetPakchunkIDFromPakFile(Entry.Filename);
+				int32 ChunkID = ParseChunkIDFromFilename(Entry.Filename);
 				if (ChunkID != INDEX_NONE)
 				{
 					ChunksToNotify.Add(ChunkID);
@@ -6945,3 +6841,5 @@ void FPakPlatformFile::MakeUniquePakFilesForTheseFiles(TArray<TArray<FString>> I
 
 
 IMPLEMENT_MODULE(FPakFileModule, PakFile);
+
+PRAGMA_ENABLE_OPTIMIZATION_ACTUAL
\ No newline at end of file
