Fixing namings of pakchunk index and chunk id.
Name the number in pakchunk file pakchunk index and the number of a platform chunk chunk id to avoid confusion.  Also keep this naming convention consistent across all platform code.

diff --git a/Engine/Source/Runtime/Core/Public/GenericPlatform/GenericPlatformChunkInstall.h b/Engine/Source/Runtime/Core/Public/GenericPlatform/GenericPlatformChunkInstall.h
--- a/Engine/Source/Runtime/Core/Public/GenericPlatform/GenericPlatformChunkInstall.h
+++ b/Engine/Source/Runtime/Core/Public/GenericPlatform/GenericPlatformChunkInstall.h
@@ -123,11 +123,11 @@ public:
 	virtual ~IPlatformChunkInstall() {}
 
 	/**
-	 * Get the current location of a chunk with pakchunk id.
-	 * @param ChunkID		The id of the pak chunk.
+	 * Get the current location of a chunk with pakchunk index.
+	 * @param PakchunkIndex	The id of the pak chunk.
 	 * @return				Enum specifying whether the chunk is available to use, waiting to install, or does not exist.
 	 **/
-	virtual EChunkLocation::Type GetPakchunkLocation( uint32 PakchunkID ) = 0;
+	virtual EChunkLocation::Type GetPakchunkLocation( int32 PakchunkIndex) = 0;
 
 	/** 
 	 * Check if a given reporting type is supported.
@@ -138,11 +138,11 @@ public:
 
 	/**
 	 * Get the current install progress of a chunk.  Let the user specify report type for platforms that support more than one.
-	 * @param PakchunkID		The id of the chunk to check.
+	 * @param ChunkID		The id of the chunk to check.
 	 * @param ReportType	The type of progress report you want.
 	 * @return				A value whose meaning is dependent on the ReportType param.
 	 **/
-	virtual float GetChunkProgress( uint32 PakchunkID, EChunkProgressReportingType::Type ReportType ) = 0;
+	virtual float GetChunkProgress( uint32 ChunkID, EChunkProgressReportingType::Type ReportType ) = 0;
 
 	/**
 	 * Inquire about the priority of chunk installation vs. game IO.
@@ -158,11 +158,11 @@ public:
 	
 	/**
 	 * Hint to the installer that we would like to prioritize a specific chunk
-	 * @param PakchunkID	The id of the pakchunk to prioritize.
+	 * @param PakchunkIndex	The index of the pakchunk to prioritize.
 	 * @param Priority		The priority for the chunk.
 	 * @return				false if the operation is not allowed or the chunk doesn't exist, otherwise true.
 	 **/
-	virtual bool PrioritizePakchunk( uint32 PakchunkID, EChunkPriority::Type Priority ) = 0;
+	virtual bool PrioritizePakchunk( int32 PakchunkIndex, EChunkPriority::Type Priority ) = 0;
 
 	/**
 	 * For platforms that support emulation of the Chunk install.  Starts transfer of the next chunk.
@@ -226,7 +226,7 @@ public:
 protected:
 		/**
 		 * Get the current location of a chunk.
-		 * Pakchunk index and OS chunk id are not the same.  Call GetPakchunkLocation instead of calling from outside.
+		 * Pakchunk index and platform chunk id are not always the same.  Call GetPakchunkLocation instead of calling from outside.
 		 * @param ChunkID		The id of the chunk to check.
 		 * @return				Enum specifying whether the chunk is available to use, waiting to install, or does not exist.
 		 **/
@@ -249,14 +249,14 @@ PRAGMA_DISABLE_DEPRECATION_WARNINGS
 class CORE_API FGenericPlatformChunkInstall : public IPlatformChunkInstall
 {
 public:
-	virtual EChunkLocation::Type GetPakchunkLocation( uint32 PakchunkID ) override
+	virtual EChunkLocation::Type GetPakchunkLocation( int32 PakchunkIndex ) override
 	{
-		return GetChunkLocation(PakchunkID);
+		return GetChunkLocation(PakchunkIndex);
 	}
 
-	virtual bool PrioritizePakchunk(uint32 PakchunkID, EChunkPriority::Type Priority)
+	virtual bool PrioritizePakchunk(int32 PakchunkIndex, EChunkPriority::Type Priority)
 	{
-		return PrioritizeChunk(PakchunkID, Priority);
+		return PrioritizeChunk(PakchunkIndex, Priority);
 	}
 
 	virtual bool GetProgressReportingTypeSupported(EChunkProgressReportingType::Type ReportType) override
@@ -269,7 +269,7 @@ public:
 		return false;
 	}
 
-	virtual float GetChunkProgress( uint32 PakchunkID, EChunkProgressReportingType::Type ReportType ) override
+	virtual float GetChunkProgress( uint32 ChunkID, EChunkProgressReportingType::Type ReportType ) override
 	{
 		if (ReportType == EChunkProgressReportingType::PercentageComplete)
 		{
@@ -288,7 +288,7 @@ public:
 		return false;
 	}
 	
-	virtual bool PrioritizeChunk( uint32 PakchunkID, EChunkPriority::Type Priority ) override
+	virtual bool PrioritizeChunk( uint32 ChunkID, EChunkPriority::Type Priority ) override
 	{
 		return false;
 	}
diff --git a/Engine/Source/Runtime/PakFile/Private/IPlatformFilePak.cpp b/Engine/Source/Runtime/PakFile/Private/IPlatformFilePak.cpp
--- a/Engine/Source/Runtime/PakFile/Private/IPlatformFilePak.cpp
+++ b/Engine/Source/Runtime/PakFile/Private/IPlatformFilePak.cpp
@@ -56,7 +56,7 @@ CSV_DECLARE_CATEGORY_MODULE_EXTERN(CORE_API, FileIO);
 
 static FString GMountStartupPaksWildCard = TEXT(MOUNT_STARTUP_PAKS_WILDCARD);
 
-int32 GetPakchunkIDFromPakFile(const FString& InFilename)
+int32 GetPakchunkIndexFromPakFile(const FString& InFilename)
 {
 	FString ChunkIdentifier(TEXT("pakchunk"));
 	FString BaseFilename = FPaths::GetBaseFilename(InFilename);
@@ -4773,7 +4773,7 @@ FPakFile::FPakFile(const TCHAR* Filename, bool bIsSigned)
 	, bSigned(bIsSigned)
 	, bIsValid(false)
 	, bFilenamesRemoved(false)
-	, ChunkID(GetPakchunkIDFromPakFile(Filename))
+	, PakchunkIndex(GetPakchunkIndexFromPakFile(Filename))
 	, bAttemptedPakEntryShrink(false)
 	, bAttemptedPakFilenameUnload(false)
  	, MappedFileHandle(nullptr)
@@ -4803,7 +4803,7 @@ FPakFile::FPakFile(IPlatformFile* LowerLevel, const TCHAR* Filename, bool bIsSig
 	, bSigned(bIsSigned)
 	, bIsValid(false)
 	, bFilenamesRemoved(false)
-	, ChunkID(GetPakchunkIDFromPakFile(Filename))
+	, PakchunkIndex(GetPakchunkIndexFromPakFile(Filename))
 	, bAttemptedPakEntryShrink(false)
 	, bAttemptedPakFilenameUnload(false)
 	, MappedFileHandle(nullptr)
@@ -4830,7 +4830,7 @@ FPakFile::FPakFile(FArchive* Archive)
 	, bSigned(false)
 	, bIsValid(false)
 	, bFilenamesRemoved(false)
-	, ChunkID(INDEX_NONE)
+	, PakchunkIndex(INDEX_NONE)
 	, MappedFileHandle(nullptr)
 	, CacheType(FPakFile::ECacheType::Shared)
 	, CacheIndex(-1)
@@ -6092,10 +6092,10 @@ void FPakPlatformFile::FindPakFilesInDirectory(IPlatformFile* LowLevelFile, cons
 					// if a platform supports chunk style installs, make sure that the chunk a pak file resides in is actually fully installed before accepting pak files from it
 					if (ChunkInstall)
 					{
-						int32 PakchunkID = GetPakchunkIDFromPakFile(Filename);
-						if (PakchunkID != INDEX_NONE)
+						int32 PakchunkIndex = GetPakchunkIndexFromPakFile(Filename);
+						if (PakchunkIndex != INDEX_NONE)
 						{
-							if (ChunkInstall->GetPakchunkLocation(PakchunkID) == EChunkLocation::NotAvailable)
+							if (ChunkInstall->GetPakchunkLocation(PakchunkIndex) == EChunkLocation::NotAvailable)
 							{
 								return true;
 							}
@@ -6341,7 +6341,7 @@ bool FPakPlatformFile::Mount(const TCHAR* InPakFilename, uint32 PakOrder, const
 				Entry.Path = InPath;
 				Entry.ReadOrder = PakOrder;
 				Entry.EncryptionKeyGuid = Pak->GetInfo().EncryptionKeyGuid;
-				Entry.ChunkID = Pak->ChunkID;
+				Entry.PakchunkIndex = Pak->PakchunkIndex;
 
 				delete Pak;
 				PakHandle.Reset();
@@ -6358,7 +6358,7 @@ bool FPakPlatformFile::Mount(const TCHAR* InPakFilename, uint32 PakOrder, const
 			PRAGMA_DISABLE_DEPRECATION_WARNINGS
 			FCoreDelegates::PakFileMountedCallback.Broadcast(InPakFilename);
 			PRAGMA_ENABLE_DEPRECATION_WARNINGS
-			FCoreDelegates::OnPakFileMounted.Broadcast(InPakFilename, Pak->ChunkID);
+			FCoreDelegates::OnPakFileMounted.Broadcast(InPakFilename, Pak->PakchunkIndex);
 			if (FCoreDelegates::NewFileAddedDelegate.IsBound())
 			{
 				TArray<FString> Filenames;
@@ -6599,10 +6599,10 @@ void FPakPlatformFile::RegisterEncryptionKey(const FGuid& InGuid, const FAES::FA
 				UE_LOG(LogPakFile, Log, TEXT("Successfully mounted deferred pak file '%s'"), *Entry.Filename);
 				NumMounted++;
 
-				int32 ChunkID = GetPakchunkIDFromPakFile(Entry.Filename);
-				if (ChunkID != INDEX_NONE)
+				int32 PakchunkIndex = GetPakchunkIndexFromPakFile(Entry.Filename);
+				if (PakchunkIndex != INDEX_NONE)
 				{
-					ChunksToNotify.Add(ChunkID);
+					ChunksToNotify.Add(PakchunkIndex);
 				}
 			}
 			else
@@ -6617,9 +6617,9 @@ void FPakPlatformFile::RegisterEncryptionKey(const FGuid& InGuid, const FAES::FA
 		IPlatformChunkInstall * ChunkInstall = FPlatformMisc::GetPlatformChunkInstall();
 		if (ChunkInstall)
 		{
-			for (int32 ChunkID : ChunksToNotify)
+			for (int32 PakchunkIndex : ChunksToNotify)
 			{
-				ChunkInstall->ExternalNotifyChunkAvailable(ChunkID);
+				ChunkInstall->ExternalNotifyChunkAvailable(PakchunkIndex);
 			}
 		}
 
@@ -6674,13 +6674,13 @@ void FPakPlatformFile::SetMountStartupPaksWildCard(const FString& WildCard)
 }
 
 
-EChunkLocation::Type FPakPlatformFile::GetPakChunkLocation(int32 InChunkID) const
+EChunkLocation::Type FPakPlatformFile::GetPakChunkLocation(int32 InPakchunkIndex) const
 {
 	FScopeLock ScopedLock(&PakListCritical);
 
 	for (const FPakListEntry& PakEntry : PakFiles)
 	{
-		if (PakEntry.PakFile->ChunkID == InChunkID)
+		if (PakEntry.PakFile->PakchunkIndex == InPakchunkIndex)
 		{
 			return EChunkLocation::LocalFast;
 		}
@@ -6688,7 +6688,7 @@ EChunkLocation::Type FPakPlatformFile::GetPakChunkLocation(int32 InChunkID) cons
 
 	for (const FPakListDeferredEntry& PendingPak : PendingEncryptedPakFiles)
 	{
-		if (PendingPak.ChunkID == InChunkID)
+		if (PendingPak.PakchunkIndex == InPakchunkIndex)
 		{
 			return EChunkLocation::NotAvailable;
 		}
@@ -6703,7 +6703,7 @@ bool FPakPlatformFile::AnyChunksAvailable() const
 
 	for (const FPakListEntry& PakEntry : PakFiles)
 	{
-		if (PakEntry.PakFile->ChunkID != INDEX_NONE)
+		if (PakEntry.PakFile->PakchunkIndex != INDEX_NONE)
 		{
 			return true;
 		}
@@ -6711,7 +6711,7 @@ bool FPakPlatformFile::AnyChunksAvailable() const
 
 	for (const FPakListDeferredEntry& PendingPak : PendingEncryptedPakFiles)
 	{
-		if (PendingPak.ChunkID != INDEX_NONE)
+		if (PendingPak.PakchunkIndex != INDEX_NONE)
 		{
 			return true;
 		}
diff --git a/Engine/Source/Runtime/PakFile/Public/IPlatformFilePak.h b/Engine/Source/Runtime/PakFile/Public/IPlatformFilePak.h
--- a/Engine/Source/Runtime/PakFile/Public/IPlatformFilePak.h
+++ b/Engine/Source/Runtime/PakFile/Public/IPlatformFilePak.h
@@ -547,7 +547,7 @@ class PAKFILE_API FPakFile : FNoncopyable
 	/** True if all filenames in memory for this pak file have been hashed to a 32-bit value. Wildcard traversal is impossible when true. */
 	bool bFilenamesRemoved;
 	/** ID for the chunk this pakfile is part of. INDEX_NONE if this isn't a pak chunk (derived from filename) */
-	int32 ChunkID;
+	int32 PakchunkIndex;
 	/** Flag to say we tried shrinking pak entries already */
 	bool bAttemptedPakEntryShrink;
 	/** Flag to say we tried unloading pak index filenames already */
@@ -1415,7 +1415,7 @@ class PAKFILE_API FPakPlatformFile : public IPlatformFile
 		FString Path;
 		uint32 ReadOrder;
 		FGuid EncryptionKeyGuid;
-		int32 ChunkID;
+		int32 PakchunkIndex;
 	};
 	
 	/** Wrapped file */
@@ -1565,9 +1565,9 @@ public:
 	static void SetMountStartupPaksWildCard(const FString& WildCard);
 
 	/**
-	* Determine location information for a given chunk ID. Will be DoesNotExist if the pak file wasn't detected, NotAvailable if it exists but hasn't been mounted due to a missing encryption key, or LocalFast if it exists and has been mounted
+	* Determine location information for a given pakchunk index. Will be DoesNotExist if the pak file wasn't detected, NotAvailable if it exists but hasn't been mounted due to a missing encryption key, or LocalFast if it exists and has been mounted
 	*/
-	EChunkLocation::Type GetPakChunkLocation(int32 InChunkID) const;
+	EChunkLocation::Type GetPakChunkLocation(int32 InPakchunkIndex) const;
 
 	/**
 	* Returns true if any of the mounted or pending pak files are chunks (filenames starting pakchunkN)
