Copying Private-LoadTimes-4.24 stream to Main. Biggest changes are in Materials/Shader memory freezing.

diff --git a/Engine/Source/Runtime/PakFile/Private/IPlatformFilePak.cpp b/Engine/Source/Runtime/PakFile/Private/IPlatformFilePak.cpp
--- a/Engine/Source/Runtime/PakFile/Private/IPlatformFilePak.cpp
+++ b/Engine/Source/Runtime/PakFile/Private/IPlatformFilePak.cpp
@@ -34,11 +34,17 @@
 #include "Async/MappedFileHandle.h"
 #include "IO/IoDispatcher.h"
 
+#include "ProfilingDebugging/LoadTimeTracker.h"
+
 DEFINE_LOG_CATEGORY(LogPakFile);
 
 DEFINE_STAT(STAT_PakFile_Read);
 DEFINE_STAT(STAT_PakFile_NumOpenHandles);
 
+IMPLEMENT_TYPE_LAYOUT(FPakCompressedBlock);
+IMPLEMENT_TYPE_LAYOUT(FPakEntry);
+IMPLEMENT_TYPE_LAYOUT(FPakFileData);
+
 CSV_DECLARE_CATEGORY_MODULE_EXTERN(CORE_API, FileIO);
 
 #ifndef DISABLE_NONUFS_INI_WHEN_COOKED
@@ -4244,7 +4250,7 @@ public:
 IAsyncReadFileHandle* FPakPlatformFile::OpenAsyncRead(const TCHAR* Filename)
 {
 	CSV_SCOPED_TIMING_STAT(FileIO, PakOpenAsyncRead);
-	check(GConfig);
+//	check(GConfig);
 #if USE_PAK_PRECACHE
 	if (FPlatformProcess::SupportsMultithreading() && GPakCache_Enable > 0)
 	{
@@ -4919,6 +4925,8 @@ void FPakFile::Initialize(FArchive* Reader, bool bLoadIndex)
 		{
 			Reader->Seek(FileInfoPos);
 
+			SCOPED_BOOT_TIMING("PakFile_SerilizeTrailer");
+
 			// Serialize trailer and check if everything is as expected.
 			Info.Serialize(*Reader, CompatibleVersion);
 			if (Info.Magic == FPakInfo::PakFile_Magic)
@@ -4932,7 +4940,7 @@ void FPakFile::Initialize(FArchive* Reader, bool bLoadIndex)
 	{
 		UE_CLOG(Info.Magic != FPakInfo::PakFile_Magic, LogPakFile, Fatal, TEXT("Trailing magic number (%ud) in '%s' is different than the expected one. Verify your installation."), Info.Magic, *PakFilename);
 		UE_CLOG(!(Info.Version >= FPakInfo::PakFile_Version_Initial && Info.Version <= CompatibleVersion), LogPakFile, Fatal, TEXT("Invalid pak file version (%d) in '%s'. Verify your installation."), Info.Version, *PakFilename);
-		UE_CLOG(!(Info.IndexOffset >= 0 && Info.IndexOffset < CachedTotalSize), LogPakFile, Fatal, TEXT("Index offset for pak file '%s' is invalid (%lld)"), *PakFilename, Info.IndexOffset);
+		UE_CLOG(!(Info.IndexOffset >= 0 && Info.IndexOffset < CachedTotalSize), LogPakFile, Fatal, TEXT("Index offset for pak file '%s' is invalid (%lld is bigger than file size %lld)"), *PakFilename, Info.IndexOffset, CachedTotalSize);
 		UE_CLOG(!((Info.IndexOffset + Info.IndexSize) >= 0 && (Info.IndexOffset + Info.IndexSize) <= CachedTotalSize), LogPakFile, Fatal, TEXT("Index end offset for pak file '%s' is invalid (%lld)"), *PakFilename, Info.IndexOffset + Info.IndexSize);
 
 		// If we aren't using a dynamic encryption key, process the pak file using the embedded key
@@ -4962,136 +4970,168 @@ void FPakFile::LoadIndex(FArchive* Reader)
 	}
 	else
 	{
-		// Load index into memory first.
+		if (Info.Version >= FPakInfo::PakFile_Version_FrozenIndex && Info.bIndexIsFrozen)
+		{
+			SCOPED_BOOT_TIMING("PakFile_LoadFrozen");
 
-		// If we encounter a corrupt index, try again but gather some extra debugging information so we can try and understand where it failed.
-		bool bFirstPass = true;
-		TArray<uint8> IndexData;
+			// read frozen data
+			Reader->Seek(Info.IndexOffset);
+			int32 FrozenSize = Info.IndexSize;
 
-		while (true)
+			// read in the index, etc data in one lump
+			void* DataMemory = FMemory::Malloc(FrozenSize);
+			Reader->Serialize(DataMemory, FrozenSize);
+			Data = TUniquePtr<FPakFileData>((FPakFileData*)DataMemory);
+
+			// cache the number of entries
+			NumEntries = Data->Files.Num();
+			// @todo loadtime: it is nice to serialize the mountpoint right into the Data so that IndexSize is right here
+			// but it takes this to copy it out, because it's too painful for the string manipulation when dealing with
+			// MemoryImageString everywhere MountPoint is used
+			MountPoint = Data->MountPoint;
+		}
+		else
 		{
-			Reader->Seek(Info.IndexOffset);
-			IndexData.Empty(); // The next SetNum makes this Empty() logically redundant, but we want to try and force a memory reallocation for the re-attempt
-			IndexData.SetNum(Info.IndexSize);
-			Reader->Serialize(IndexData.GetData(), Info.IndexSize);
+			// Load index into memory first.
 
-			FSHAHash EncryptedDataHash;
-			if (!bFirstPass)
-			{
-				FSHA1::HashBuffer(IndexData.GetData(), IndexData.Num(), EncryptedDataHash.Hash);
-			}
+			// If we encounter a corrupt index, try again but gather some extra debugging information so we can try and understand where it failed.
+			bool bFirstPass = true;
+			TArray<uint8> IndexData;
 
-			// Decrypt if necessary
-			if (Info.bEncryptedIndex)
-			{
-				DecryptData(IndexData.GetData(), Info.IndexSize, Info.EncryptionKeyGuid);
-			}
-			// Check SHA1 value.
-			FSHAHash ComputedHash;
-			FSHA1::HashBuffer(IndexData.GetData(), IndexData.Num(), ComputedHash.Hash);
-			if (Info.IndexHash != ComputedHash)
+			while (true)
 			{
-				if (bFirstPass)
+				Reader->Seek(Info.IndexOffset);
+				IndexData.Empty(); // The next SetNum makes this Empty() logically redundant, but we want to try and force a memory reallocation for the re-attempt
+				IndexData.SetNum(Info.IndexSize);
 				{
-					UE_LOG(LogPakFile, Log, TEXT("Corrupt pak index detected!"));
-					UE_LOG(LogPakFile, Log, TEXT(" Filename: %s"), *PakFilename);
-					UE_LOG(LogPakFile, Log, TEXT(" Encrypted: %d"), Info.bEncryptedIndex);
-					UE_LOG(LogPakFile, Log, TEXT(" Total Size: %d"), Reader->TotalSize());
-					UE_LOG(LogPakFile, Log, TEXT(" Index Offset: %d"), Info.IndexOffset);
-					UE_LOG(LogPakFile, Log, TEXT(" Index Size: %d"), Info.IndexSize);
-					UE_LOG(LogPakFile, Log, TEXT(" Stored Index Hash: %s"), *Info.IndexHash.ToString());
-					UE_LOG(LogPakFile, Log, TEXT(" Computed Index Hash [Pass 0]: %s"), *ComputedHash.ToString());
-					bFirstPass = false;
+					SCOPED_BOOT_TIMING("PakFile_LoadIndex");
+					Reader->Serialize(IndexData.GetData(), Info.IndexSize);
 				}
-				else
+
+				FSHAHash EncryptedDataHash;
+				if (!bFirstPass)
+				{
+					FSHA1::HashBuffer(IndexData.GetData(), IndexData.Num(), EncryptedDataHash.Hash);
+				}
+
+				// Decrypt if necessary
+				if (Info.bEncryptedIndex)
 				{
-					UE_LOG(LogPakFile, Log, TEXT(" Computed Index Hash [Pass 1]: %s"), *ComputedHash.ToString());
-					UE_LOG(LogPakFile, Log, TEXT(" Encrypted Index Hash: %s"), *EncryptedDataHash.ToString());
-
-					// Compute an SHA1 hash of the whole file so we can tell if the file was modified on disc (obviously as long as this isn't an  IO bug that gives us the same bogus data again)
-					FSHA1 FileHash;
-					Reader->Seek(0);
-					int64 Remaining = Reader->TotalSize();
-					TArray<uint8> WorkingBuffer;
-					WorkingBuffer.SetNum(64 * 1024);
-					while (Remaining > 0)
+					DecryptData(IndexData.GetData(), Info.IndexSize, Info.EncryptionKeyGuid);
+				}
+				// Check SHA1 value.
+				FSHAHash ComputedHash;
+				{
+					SCOPED_BOOT_TIMING("PakFile_Hash");
+					FSHA1::HashBuffer(IndexData.GetData(), IndexData.Num(), ComputedHash.Hash);
+				}
+				if (Info.IndexHash != ComputedHash)
+				{
+					if (bFirstPass)
 					{
-						int64 ToProcess = FMath::Min((int64)WorkingBuffer.Num(), Remaining);
-						Reader->Serialize(WorkingBuffer.GetData(), ToProcess);
-						FileHash.Update(WorkingBuffer.GetData(), ToProcess);
-						Remaining -= ToProcess;
+						UE_LOG(LogPakFile, Log, TEXT("Corrupt pak index detected!"));
+						UE_LOG(LogPakFile, Log, TEXT(" Filename: %s"), *PakFilename);
+						UE_LOG(LogPakFile, Log, TEXT(" Encrypted: %d"), Info.bEncryptedIndex);
+						UE_LOG(LogPakFile, Log, TEXT(" Total Size: %d"), Reader->TotalSize());
+						UE_LOG(LogPakFile, Log, TEXT(" Index Offset: %d"), Info.IndexOffset);
+						UE_LOG(LogPakFile, Log, TEXT(" Index Size: %d"), Info.IndexSize);
+						UE_LOG(LogPakFile, Log, TEXT(" Stored Index Hash: %s"), *Info.IndexHash.ToString());
+						UE_LOG(LogPakFile, Log, TEXT(" Computed Index Hash [Pass 0]: %s"), *ComputedHash.ToString());
+						bFirstPass = false;
 					}
+					else
+					{
+						UE_LOG(LogPakFile, Log, TEXT(" Computed Index Hash [Pass 1]: %s"), *ComputedHash.ToString());
+						UE_LOG(LogPakFile, Log, TEXT(" Encrypted Index Hash: %s"), *EncryptedDataHash.ToString());
+
+						// Compute an SHA1 hash of the whole file so we can tell if the file was modified on disc (obviously as long as this isn't an  IO bug that gives us the same bogus data again)
+						FSHA1 FileHash;
+						Reader->Seek(0);
+						int64 Remaining = Reader->TotalSize();
+						TArray<uint8> WorkingBuffer;
+						WorkingBuffer.SetNum(64 * 1024);
+						while (Remaining > 0)
+						{
+							int64 ToProcess = FMath::Min((int64)WorkingBuffer.Num(), Remaining);
+							Reader->Serialize(WorkingBuffer.GetData(), ToProcess);
+							FileHash.Update(WorkingBuffer.GetData(), ToProcess);
+							Remaining -= ToProcess;
+						}
 
-					FileHash.Final();
-					FSHAHash FinalFileHash;
-					FileHash.GetHash(FinalFileHash.Hash);
-					UE_LOG(LogPakFile, Log, TEXT(" File Hash: %s"), *FinalFileHash.ToString());
+						FileHash.Final();
+						FSHAHash FinalFileHash;
+						FileHash.GetHash(FinalFileHash.Hash);
+						UE_LOG(LogPakFile, Log, TEXT(" File Hash: %s"), *FinalFileHash.ToString());
 
-					UE_LOG(LogPakFile, Fatal, TEXT("Corrupted index in pak file (SHA hash mismatch)."));
+						UE_LOG(LogPakFile, Fatal, TEXT("Corrupted index in pak file (SHA hash mismatch)."));
+					}
 				}
-			}
-			else
-			{
-				if (!bFirstPass)
+				else
 				{
-					UE_LOG(LogPakFile, Log, TEXT("Pak index corruption appears to have recovered on the second attempt!"));
+					if (!bFirstPass)
+					{
+						UE_LOG(LogPakFile, Log, TEXT("Pak index corruption appears to have recovered on the second attempt!"));
+					}
+					break;
 				}
-				break;
 			}
-		}
 
-		FMemoryReader IndexReader(IndexData);
+			FMemoryReader IndexReader(IndexData);
 
-		// Read the default mount point and all entries.
-		NumEntries = 0;
-		IndexReader << MountPoint;
-		IndexReader << NumEntries;
+			// Read the default mount point and all entries.
+			NumEntries = 0;
+			IndexReader << MountPoint;
+			IndexReader << NumEntries;
 
-		MakeDirectoryFromPath(MountPoint);
-		// Allocate enough memory to hold all entries (and not reallocate while they're being added to it).
-		Files.Empty(NumEntries);
+			MakeDirectoryFromPath(MountPoint);
+			// Allocate enough memory to hold all entries (and not reallocate while they're being added to it).
 
-		for (int32 EntryIndex = 0; EntryIndex < NumEntries; EntryIndex++)
-		{
-			// Serialize from memory.
-			FPakEntry Entry;
-			FString Filename;
-			IndexReader << Filename;
-			Entry.Serialize(IndexReader, Info.Version);
+			Data = MakeUnique<FPakFileData>();
+			Data->Files.Empty(NumEntries);
+			Data->Files.AddZeroed(NumEntries);
 
-			// Add new file info.
-			Files.Add(Entry);
-
-			// Construct Index of all directories in pak file.
-			FString Path = FPaths::GetPath(Filename);
-			MakeDirectoryFromPath(Path);
-			FPakDirectory* Directory = Index.Find(Path);
-			if (Directory != NULL)
-			{
-				Directory->Add(FPaths::GetCleanFilename(Filename), EntryIndex);
-			}
-			else
+			for (int32 EntryIndex = 0; EntryIndex < NumEntries; EntryIndex++)
 			{
-				FPakDirectory& NewDirectory = Index.Add(Path);
-				NewDirectory.Add(FPaths::GetCleanFilename(Filename), EntryIndex);
+				// Serialize from memory.
+				FPakEntry& Entry = Data->Files[EntryIndex];
+				FString Filename;
+				IndexReader << Filename;
+				{
+					SCOPED_BOOT_TIMING("PakFile_SerializeEntry");
+					Entry.Serialize(IndexReader, Info.Version);
+				}
 
-				// add the parent directories up to the mount point
-				while (MountPoint != Path)
+				// Construct Index of all directories in pak file.
+				FString Path = FPaths::GetPath(Filename);
+				MakeDirectoryFromPath(Path);
+				FPakDirectory* Directory = Data->Index.Find(Path);
+				if (Directory != NULL)
 				{
-					Path.LeftInline(Path.Len() - 1, false);
-					int32 Offset = 0;
-					if (Path.FindLastChar('/', Offset))
+					Directory->Add(FPaths::GetCleanFilename(Filename), EntryIndex);
+				}
+				else
+				{
+					FPakDirectory& NewDirectory = Data->Index.Add(Path);
+					NewDirectory.Add(FPaths::GetCleanFilename(Filename), EntryIndex);
+
+					// add the parent directories up to the mount point
+					while (MountPoint != Path)
 					{
-						Path.LeftInline(Offset, false);
-						MakeDirectoryFromPath(Path);
-						if (Index.Find(Path) == NULL)
+						Path.LeftInline(Path.Len() - 1);
+						int32 Offset = 0;
+						if (Path.FindLastChar('/', Offset))
 						{
-							Index.Add(Path);
+							Path.LeftInline(Offset);
+							MakeDirectoryFromPath(Path);
+							if (Data->Index.Find(Path) == NULL)
+							{
+								Data->Index.Add(Path);
+							}
+						}
+						else
+						{
+							Path = MountPoint;
 						}
-					}
-					else
-					{
-						Path = MountPoint;
 					}
 				}
 			}
@@ -5222,6 +5262,12 @@ bool FPakFile::UnloadPakEntryFilenames(TMap<uint64, FPakEntry>& CrossPakCollisio
 	// Set this flag so if unloading fails, we don't try again
 	bAttemptedPakFilenameUnload = true;
 
+	if (Info.bIndexIsFrozen)
+	{
+		UE_LOG(LogPakFile, Warning, TEXT("FAILED unloading filenames for pak '%s' - its index was frozen and cannot be modified"), *PakFilename);
+		return false;
+	}
+
 	// Variables for the filename hashing and collision detection.
 	int NumRetries = 0;
 	const int MAX_RETRIES = bAllowRetries ? 10 : 1;
@@ -5245,7 +5291,7 @@ bool FPakFile::UnloadPakEntryFilenames(TMap<uint64, FPakEntry>& CrossPakCollisio
 		int32 EntryIndex = 0;
 		FString FinalFilename;
 		FinalFilename.Reset(1024);
-		for (TMap<FString, FPakDirectory>::TConstIterator It(Index); (It && !bHasCollision); ++It)
+		for (TMemoryImageMap<FMemoryImageString, FPakDirectory>::TConstIterator It(Data->Index); (It && !bHasCollision); ++It)
 		{
 			for (FPakDirectory::TConstIterator DirectoryIt(It.Value()); DirectoryIt; ++DirectoryIt)
 			{
@@ -5260,7 +5306,7 @@ bool FPakFile::UnloadPakEntryFilenames(TMap<uint64, FPakEntry>& CrossPakCollisio
 
 				const FPakEntry* EntryFromPreviousPaks = CrossPakCollisionChecker.Find(FilenameHash);
 				const FPakEntry* EntryFromCurrentPak = NewCollisionCheckEntries.Find(FilenameHash);
-				const FPakEntry& CurrentEntry = Files[DirectoryIt.Value()];
+				const FPakEntry& CurrentEntry = Data->Files[DirectoryIt.Value()];
 
 				if (EntryFromPreviousPaks && (FMemory::Memcmp(EntryFromPreviousPaks->Hash, CurrentEntry.Hash, sizeof(CurrentEntry.Hash))) != 0)
 				{
@@ -5390,7 +5436,7 @@ bool FPakFile::UnloadPakEntryFilenames(TMap<uint64, FPakEntry>& CrossPakCollisio
 	// Clear out those portions of the Index allowed by the user.
 	if (DirectoryRootsToKeep != nullptr)
 	{
-		for(auto It = Index.CreateIterator(); It; ++It)
+		for(auto It = Data->Index.CreateIterator(); It; ++It)
 		{
 			const FString DirectoryName = MountPoint / It.Key();
 
@@ -5410,7 +5456,7 @@ bool FPakFile::UnloadPakEntryFilenames(TMap<uint64, FPakEntry>& CrossPakCollisio
 			}
 		}
 
-		Index.Shrink();
+		Data->Index.Shrink();
 
 #if defined(FPAKFILE_UNLOADPAKENTRYFILENAMES_LOGKEPTFILENAMES)
 		for (TMap<FString, FPakDirectory>::TConstIterator It(Index); It; ++It)
@@ -5421,7 +5467,7 @@ bool FPakFile::UnloadPakEntryFilenames(TMap<uint64, FPakEntry>& CrossPakCollisio
 	}
 	else
 	{
-		Index.Empty(0);
+		Data->Index.Empty(0);
 	}
 
 	return true;
@@ -5442,13 +5488,20 @@ bool FPakFile::ShrinkPakEntriesMemoryUsage()
 	// Set this flag so if shrinking fails, we don't try again
 	bAttemptedPakEntryShrink = true;
 
+	// if the index was frozen, we can't unload parts of it, so skip on out
+	if (Info.bIndexIsFrozen)
+	{
+		UE_LOG(LogPakFile, Warning, TEXT("FAILED shrinking entries for pak file '%s' - its index was frozen and cannot be modified"), *PakFilename);
+		return false;
+	}
+
 	// Wander every file entry.
 	int TotalSizeOfCompressedEntries = 0;
 	bool bIsPossibleToShrink = true;
 	int32 EntryIndex = 0;
 	for (EntryIndex = 0; EntryIndex < NumEntries; ++EntryIndex)
 	{
-		FPakEntry& Entry = Files[EntryIndex];
+		FPakEntry& Entry = Data->Files[EntryIndex];
 
 		bool bIsOffset32BitSafe = Entry.Offset <= MAX_uint32;
 		bool bIsSize32BitSafe = Entry.Size <= MAX_uint32;
@@ -5538,7 +5591,7 @@ bool FPakFile::ShrinkPakEntriesMemoryUsage()
 	uint8* CurrentEntryPtr = MiniPakEntries;
 	for (EntryIndex = 0; EntryIndex < NumEntries; ++EntryIndex)
 	{
-		FPakEntry* FullEntry = &Files[EntryIndex];
+		FPakEntry* FullEntry = &Data->Files[EntryIndex];
 
 		MiniPakEntriesOffsets[EntryIndex] = CurrentEntryPtr - MiniPakEntries;
 
@@ -5650,7 +5703,7 @@ bool FPakFile::ShrinkPakEntriesMemoryUsage()
 
 	// Clear out the Files data. We compressed it, and we don't need the wasted
 	// space of the original anymore.
-	Files.Empty(0);
+	Data->Files.Empty(0);
 
 	static uint64 Total = 0;
 	Total += TotalSizeOfCompressedEntries;
@@ -5709,7 +5762,7 @@ public:
 
 void FPakFile::GetFilenames(TArray<FString>& OutFileList) const
 {
-	for (const TMap<FString, FPakDirectory>::ElementType& DirectoryElement : Index)
+	for (const TMemoryImageMap<FMemoryImageString, FPakDirectory>::ElementType& DirectoryElement : Data->Index)
 	{
 		const  FPakDirectory& Directory = DirectoryElement.Value;
 		for (const FPakDirectory::ElementType& FileElement : Directory)
@@ -5730,7 +5783,7 @@ void FPakFile::GetFilenamesInChunk(const TArray<int32>& InChunkIDs, TArray<FStri
 		int32 ChunkEnd = ChunkStart + FPakInfo::MaxChunkDataSize;
 		int32 FileIndex = 0;
 
-		for (const FPakEntry& File : Files)
+		for (const FPakEntry& File : Data->Files)
 		{
 			int32 FileStart = File.Offset;
 			int32 FileEnd = File.Offset + File.Size;
@@ -5752,14 +5805,14 @@ void FPakFile::GetFilenamesInChunk(const TArray<int32>& InChunkIDs, TArray<FStri
 	}
 
 	int32 Remaining = OverlappingEntries.Num();
-	for (const TMap<FString, FPakDirectory>::ElementType& DirectoryElement : Index)
+	for (const TMemoryImageMap<FMemoryImageString, FPakDirectory>::ElementType& DirectoryElement : Data->Index)
 	{
 		const  FPakDirectory& Directory = DirectoryElement.Value;
 		for (const FPakDirectory::ElementType& FileElement : Directory)
 		{
 			if (OverlappingEntries.Contains(FileElement.Value))
 			{
-				OutFileList.Add(DirectoryElement.Key / FileElement.Key);
+				OutFileList.Add(FString(DirectoryElement.Key) / FString(FileElement.Key));
 				if (--Remaining == 0)
 				{
 					break;
@@ -5874,7 +5927,7 @@ FPakFile::EFindResult FPakFile::Find(const FString& Filename, FPakEntry* OutEntr
 				}
 				else
 				{
-					const FPakEntry* FoundEntry = &Files[FoundEntryIndex];
+					const FPakEntry* FoundEntry = &Data->Files[FoundEntryIndex];
 
 					bDeleted = FoundEntry->IsDeleteRecord();
 
@@ -5900,7 +5953,7 @@ FPakFile::EFindResult FPakFile::Find(const FString& Filename, FPakEntry* OutEntr
 			const FPakDirectory* PakDirectory = FindDirectory(*Path);
 			if (PakDirectory != NULL)
 			{
-				FString RelativeFilename(Filename.Mid(Path.Len() + 1));
+				FMemoryImageString RelativeFilename(Filename.Mid(Path.Len() + 1));
 				int32 const* FoundEntryIndex = PakDirectory->Find(RelativeFilename);
 				if (FoundEntryIndex != NULL)
 				{
@@ -5932,7 +5985,7 @@ FPakFile::EFindResult FPakFile::Find(const FString& Filename, FPakEntry* OutEntr
 					}
 					else
 					{
-						const FPakEntry* FoundEntry = &Files[*FoundEntryIndex];
+						const FPakEntry* FoundEntry = &Data->Files[*FoundEntryIndex];
 						bDeleted = FoundEntry->IsDeleteRecord();
 
 						if (OutEntry != NULL)
@@ -6530,6 +6583,7 @@ int32 FPakPlatformFile::MountAllPakFiles(const TArray<FString>& PakFolders, cons
 
 			UE_LOG(LogPakFile, Display, TEXT("Mounting pak file %s."), *PakFilename);
 
+			SCOPED_BOOT_TIMING("Pak_Mount");
 			if (Mount(*PakFilename, PakOrder))
 			{
 				++NumPakFilesMounted;
@@ -6890,10 +6944,10 @@ void FPakFile::AddSpecialFile(FPakEntry Entry, const FString& Filename)
 {
 	MakeDirectoryFromPath(MountPoint);
 
-	int32 EntryIndex = Files.Num();
+	int32 EntryIndex = Data->Files.Num();
 
 	// Add new file info.
-	Files.Add(Entry);
+	Data->Files.Add(Entry);
 	NumEntries++;
 
 	// Construct Index of all directories in pak file.
@@ -6907,14 +6961,14 @@ void FPakFile::AddSpecialFile(FPakEntry Entry, const FString& Filename)
 
 
 
-	FPakDirectory* Directory = Index.Find(Path); 
+	FPakDirectory* Directory = Data->Index.Find(Path);
 	if (Directory != NULL)
 	{
 		Directory->Add(FPaths::GetCleanFilename(CleanFilename), EntryIndex);
 	}
 	else
 	{
-		FPakDirectory& NewDirectory = Index.Add(Path);
+		FPakDirectory& NewDirectory = Data->Index.Add(Path);
 		NewDirectory.Add(FPaths::GetCleanFilename(CleanFilename), EntryIndex);
 
 		// add the parent directories up to the mount point
@@ -6926,9 +6980,9 @@ void FPakFile::AddSpecialFile(FPakEntry Entry, const FString& Filename)
 			{
 				Path.LeftInline(Offset, false);
 				MakeDirectoryFromPath(Path);
-				if (Index.Find(Path) == NULL)
+				if (Data->Index.Find(Path) == NULL)
 				{
-					Index.Add(Path);
+					Data->Index.Add(Path);
 				}
 			}
 			else
@@ -6958,7 +7012,7 @@ void FPakPlatformFile::MakeUniquePakFilesForTheseFiles(TArray<TArray<FString>> I
 					// we successfully mounted the file, find the empty pak file we just added.
 					for (int j = 0; j < PakFiles.Num(); j++)
 					{
-						if (PakFiles[j].PakFile->Files.Num() == 0 && PakFiles[j].PakFile->CachedTotalSize == PakFile->CachedTotalSize)
+						if (PakFiles[j].PakFile->Data->Files.Num() == 0 && PakFiles[j].PakFile->CachedTotalSize == PakFile->CachedTotalSize)
 						{
 							NewPakFile = PakFiles[j].PakFile;
 							break;
diff --git a/Engine/Source/Runtime/PakFile/Public/IPlatformFilePak.h b/Engine/Source/Runtime/PakFile/Public/IPlatformFilePak.h
--- a/Engine/Source/Runtime/PakFile/Public/IPlatformFilePak.h
+++ b/Engine/Source/Runtime/PakFile/Public/IPlatformFilePak.h
@@ -13,6 +13,7 @@
 #include "RSA.h"
 #include "Misc/SecureHash.h"
 #include "GenericPlatform/GenericPlatformChunkInstall.h"
+#include "Serialization/MemoryImage.h"
 
 class FChunkCacheWorker;
 class IAsyncReadFileHandle;
@@ -92,6 +93,7 @@ struct FPakInfo
 		PakFile_Version_DeleteRecords = 6,
 		PakFile_Version_EncryptionKeyGuid = 7,
 		PakFile_Version_FNameBasedCompressionMethod = 8,
+		PakFile_Version_FrozenIndex = 9,
 
 
 		PakFile_Version_Last,
@@ -111,6 +113,9 @@ struct FPakInfo
 	FSHAHash IndexHash;
 	/** Flag indicating if the pak index has been encrypted. */
 	uint8 bEncryptedIndex;
+	/** Flag indicating if the pak index has been frozen */
+	// @todo loadtime: we should find a way to unload the index - potentially make two indices, the full one and unloaded one? unclear how, but at least we now have an option to choose per-platform
+	uint8 bIndexIsFrozen;
 	/** Encryption key guid. Empty if we should use the embedded key. */
 	FGuid EncryptionKeyGuid;
 	/** Compression methods used in this pak file (FNames, saved as FStrings) */
@@ -125,6 +130,7 @@ struct FPakInfo
 		, IndexOffset(-1)
 		, IndexSize(0)
 		, bEncryptedIndex(0)
+		, bIndexIsFrozen(0)
 	{
 		// we always put in a NAME_None entry as index 0, so that an uncompressed PakEntry will have CompressionMethodIndex of 0 and can early out easily
 		CompressionMethods.Add(NAME_None);
@@ -140,6 +146,7 @@ struct FPakInfo
 		int64 Size = sizeof(Magic) + sizeof(Version) + sizeof(IndexOffset) + sizeof(IndexSize) + sizeof(IndexHash) + sizeof(bEncryptedIndex);
 		if (InVersion >= PakFile_Version_EncryptionKeyGuid) Size += sizeof(EncryptionKeyGuid);
 		if (InVersion >= PakFile_Version_FNameBasedCompressionMethod) Size += CompressionMethodNameLen * MaxNumCompressionMethods;
+		if (InVersion >= PakFile_Version_FrozenIndex) Size += sizeof(bIndexIsFrozen);
 
 		return Size;
 	}
@@ -195,6 +202,11 @@ struct FPakInfo
 			}
 		}
 
+		if (Version >= PakFile_Version_FrozenIndex)
+		{
+			Ar << bIndexIsFrozen;
+		}
+
 		if (Version < PakFile_Version_FNameBasedCompressionMethod)
 		{
 			// for old versions, put in some known names that we may have used
@@ -264,10 +276,12 @@ struct FPakInfo
  */
 struct FPakCompressedBlock
 {
+	DECLARE_TYPE_LAYOUT(FPakCompressedBlock, NonVirtual);
+
 	/** Offset of the start of a compression block. Offset is relative to the start of the compressed chunk data */
-	int64 CompressedStart;
+	LAYOUT_FIELD(int64, CompressedStart);
 	/** Offset of the end of a compression block. This may not align completely with the start of the next block. Offset is relative to the start of the compressed chunk data. */
-	int64 CompressedEnd;
+	LAYOUT_FIELD(int64, CompressedEnd);
 
 	bool operator == (const FPakCompressedBlock& B) const
 	{
@@ -301,28 +315,30 @@ FORCEINLINE FArchive& operator<<(FArchive& Ar, FPakCompressedBlock& Block)
  */
 struct FPakEntry
 {
+	DECLARE_TYPE_LAYOUT(FPakEntry, NonVirtual);
+
 	static const uint8 Flag_None = 0x00;
 	static const uint8 Flag_Encrypted = 0x01;
 	static const uint8 Flag_Deleted = 0x02;
 
 	/** Offset into pak file where the file is stored.*/
-	int64 Offset;
+	LAYOUT_FIELD(int64, Offset);
 	/** Serialized file size. */
-	int64 Size;
+	LAYOUT_FIELD(int64, Size);
 	/** Uncompressed file size. */
-	int64 UncompressedSize;
+	LAYOUT_FIELD(int64, UncompressedSize);
 	/** File SHA1 value. */
-	uint8 Hash[20];
+	LAYOUT_ARRAY(uint8, Hash, 20);
 	/** Array of compression blocks that describe how to decompress this pak entry. */
-	TArray<FPakCompressedBlock> CompressionBlocks;
+	LAYOUT_FIELD(TMemoryImageArray<FPakCompressedBlock>, CompressionBlocks);
 	/** Size of a compressed block in the file. */
-	uint32 CompressionBlockSize;
+	LAYOUT_FIELD(uint32, CompressionBlockSize);
 	/** Index into the compression methods in this pakfile. */
-	uint32 CompressionMethodIndex;
+	LAYOUT_FIELD(uint32, CompressionMethodIndex);
 	/** Pak entry flags. */
-	uint8 Flags;
+	LAYOUT_FIELD(uint8, Flags);
 	/** Flag is set to true when FileHeader has been checked against PakHeader. It is not serialized. */
-	mutable bool  Verified;
+	LAYOUT_MUTABLE_FIELD(bool, Verified);
 
 	/**
 	 * Constructor.
@@ -496,7 +512,20 @@ struct FPakEntry
 };
 
 /** Pak directory type mapping a filename to a FPakEntry index within the FPakFile.Files array or the MiniPakEntriesOffsets array depending on the bit-encoded state of the FPakEntry structures. */
-typedef TMap<FString, int32> FPakDirectory;
+typedef TMemoryImageMap<FMemoryImageString, int32> FPakDirectory;
+
+struct PAKFILE_API FPakFileData
+{
+	DECLARE_TYPE_LAYOUT(FPakFileData, NonVirtual);
+
+
+	/** Mount point. */
+	LAYOUT_FIELD(FMemoryImageString, MountPoint);
+	/** Info on all files stored in pak. */
+	LAYOUT_FIELD(TMemoryImageArray<FPakEntry>, Files);
+	/** Pak Index organized as a map of directories for faster Directory iteration. Completely valid only when bFilenamesRemoved == false, although portions may still be valid after a call to UnloadPakEntryFilenames() while utilizing DirectoryRootsToKeep. */
+	LAYOUT_FIELD((TMemoryImageMap<FMemoryImageString, FPakDirectory>), Index);
+};
 
 /**
  * Pak file.
@@ -505,6 +534,8 @@ class PAKFILE_API FPakFile : FNoncopyable
 {
 	friend class FPakPlatformFile;
 
+	TUniquePtr<FPakFileData> Data;
+
 	/** Pak filename. */
 	FString PakFilename;
 	FName PakFilenameName;
@@ -518,10 +549,6 @@ class PAKFILE_API FPakFile : FNoncopyable
 	FPakInfo Info;
 	/** Mount point. */
 	FString MountPoint;
-	/** Info on all files stored in pak. */
-	TArray<FPakEntry> Files;	
-	/** Pak Index organized as a map of directories for faster Directory iteration. Completely valid only when bFilenamesRemoved == false, although portions may still be valid after a call to UnloadPakEntryFilenames() while utilizing DirectoryRootsToKeep. */
-	TMap<FString, FPakDirectory> Index;
 	/** The hash to use when generating a filename hash (CRC) to avoid collisions within the hashed filename space. */
 	uint64 FilenameStartHash;
 	/** An array of 256 + 1 size that represents the starting index of the most significant byte of a hash group within the FilenameHashes array. */
@@ -672,9 +699,9 @@ public:
 	 *
 	 * @return Pak index.
 	 */
-	const TMap<FString, FPakDirectory>& GetIndex() const
+	const TMemoryImageMap<FMemoryImageString, FPakDirectory>& GetIndex() const
 	{
-		return Index;
+		return Data->Index;
 	}
 
 	/**
@@ -761,7 +788,7 @@ public:
 			//checkf(!bFilenamesRemoved, TEXT("FPakFile::FindFilesAtPath() can only be used before FPakPlatformFile::UnloadFilenames() is called."));
 
 			TArray<FString> DirectoriesInPak; // List of all unique directories at path
-			for (TMap<FString, FPakDirectory>::TConstIterator It(Index); It; ++It)
+			for (TMemoryImageMap<FMemoryImageString, FPakDirectory>::TConstIterator It(Data->Index); It; ++It)
 			{
 				FString PakPath(MountPoint + It.Key());
 				// Check if the file is under the specified path.
@@ -823,7 +850,7 @@ public:
 		// Check the specified path is under the mount point of this pak file.
 		if (Directory.StartsWith(MountPoint))
 		{
-			PakDirectory = Index.Find(Directory.Mid(MountPoint.Len()));
+			PakDirectory = Data->Index.Find(Directory.Mid(MountPoint.Len()));
 		}
 		return PakDirectory;
 	}
@@ -852,7 +879,7 @@ public:
 		/** Owner pak file. */
 		const FPakFile& PakFile;
 		/** Index iterator. */
-		TMap<FString, FPakDirectory>::TConstIterator IndexIt;
+		TMemoryImageMap<FMemoryImageString, FPakDirectory>::TConstIterator IndexIt;
 		/** Directory iterator. */
 		FPakDirectory::TConstIterator DirectoryIt;
 		/** The cached filename for return in Filename() */
@@ -867,10 +894,10 @@ public:
 		 * @param InPakFile Pak file to iterate.
 		 */
 		FFileIterator(const FPakFile& InPakFile, bool bInIncludeDeleted = false )
-		:	PakFile(InPakFile)
-		, IndexIt(PakFile.GetIndex())
-		, DirectoryIt((IndexIt ? FPakDirectory::TConstIterator(IndexIt.Value()): FPakDirectory()))
-		, bIncludeDeleted(bInIncludeDeleted)
+			: PakFile(InPakFile)
+			, IndexIt(PakFile.GetIndex())
+			, DirectoryIt((IndexIt ? FPakDirectory::TConstIterator(IndexIt.Value()): FPakDirectory()))
+			, bIncludeDeleted(bInIncludeDeleted)
 		{
 			AdvanceToValid();
 			UpdateCachedFilename();
@@ -897,11 +924,11 @@ public:
 		}
 
 		const FString& Filename() const		{ return CachedFilename; }
-		const FPakEntry& Info() const	{ return PakFile.Files[DirectoryIt.Value()]; }
+		const FPakEntry& Info() const	{ return PakFile.Data->Files[DirectoryIt.Value()]; }
 		const int32 GetIndexInPakFile() const { return DirectoryIt.Value(); }
 
 	private:
-		FORCEINLINE void AdvanceToValid()
+		void AdvanceToValid()
 		{
 			SkipDeletedIfRequired();
 			while (!DirectoryIt && IndexIt)
@@ -923,7 +950,7 @@ public:
 		{
 			if (!!IndexIt && !!DirectoryIt)
 			{
-				CachedFilename = IndexIt.Key() + DirectoryIt.Key();
+				CachedFilename = FString(IndexIt.Key()) + FString(DirectoryIt.Key());
 			}
 			else
 			{
@@ -1922,7 +1949,7 @@ public:
 			{
 				for (FPakDirectory::TConstIterator DirectoryIt(*PakDirectory); DirectoryIt; ++DirectoryIt)
 				{
-					if (PakFile->Files[DirectoryIt.Value()].Offset == FileEntry.Offset)
+					if (PakFile->Data->Files[DirectoryIt.Value()].Offset == FileEntry.Offset)
 					{
 						const FString& RealFilename = DirectoryIt.Key();
 						return Path / RealFilename;
diff --git a/Engine/Source/Runtime/Projects/Private/PluginManager.cpp b/Engine/Source/Runtime/Projects/Private/PluginManager.cpp
--- a/Engine/Source/Runtime/Projects/Private/PluginManager.cpp
+++ b/Engine/Source/Runtime/Projects/Private/PluginManager.cpp
@@ -431,7 +431,7 @@ void FPluginManager::ReadAllPlugins(TMap<FString, TSharedRef<FPlugin>>& Plugins,
 				}
 			}
 
-			if (Parent->GetDescriptor().SupportedTargetPlatforms.Num() != 0 && Child->GetDescriptor().SupportedTargetPlatforms.Num() != 0)
+			if (Parent->GetDescriptor().SupportedTargetPlatforms.Num() != 0)
 			{
 				for (const FString& SupportedTargetPlatform : Child->GetDescriptor().SupportedTargetPlatforms)
 				{
diff --git a/Engine/Source/Runtime/RHI/Private/DynamicRHI.cpp b/Engine/Source/Runtime/RHI/Private/DynamicRHI.cpp
--- a/Engine/Source/Runtime/RHI/Private/DynamicRHI.cpp
+++ b/Engine/Source/Runtime/RHI/Private/DynamicRHI.cpp
@@ -15,6 +15,9 @@
 #include "GenericPlatform/GenericPlatformCrashContext.h"
 #include "PipelineStateCache.h"
 
+IMPLEMENT_TYPE_LAYOUT(FRayTracingGeometryInitializer);
+IMPLEMENT_TYPE_LAYOUT(FRayTracingGeometrySegment);
+
 #ifndef PLATFORM_ALLOW_NULL_RHI
 	#define PLATFORM_ALLOW_NULL_RHI		0
 #endif
@@ -215,7 +218,7 @@ void RHIInit(bool bHasEditorToken)
 	if(!GDynamicRHI)
 	{
 		// read in any data driven shader platform info structures we can find
-		FDataDrivenShaderPlatformInfo::Initialize();
+		FGenericDataDrivenShaderPlatformInfo::Initialize();
 
 		GRHICommandList.LatchBypass(); // read commandline for bypass flag
 
diff --git a/Engine/Source/Runtime/RHI/Private/PipelineStateCache.cpp b/Engine/Source/Runtime/RHI/Private/PipelineStateCache.cpp
--- a/Engine/Source/Runtime/RHI/Private/PipelineStateCache.cpp
+++ b/Engine/Source/Runtime/RHI/Private/PipelineStateCache.cpp
@@ -80,7 +80,7 @@ static FAutoConsoleCommand DumpPipelineCmd(
 
 void SetComputePipelineState(FRHICommandList& RHICmdList, FRHIComputeShader* ComputeShader)
 {
-	RHICmdList.SetComputePipelineState(PipelineStateCache::GetAndOrCreateComputePipelineState(RHICmdList, ComputeShader));
+	RHICmdList.SetComputePipelineState(PipelineStateCache::GetAndOrCreateComputePipelineState(RHICmdList, ComputeShader), ComputeShader);
 }
 
 extern RHI_API FRHIComputePipelineState* ExecuteSetComputePipelineState(FComputePipelineState* ComputePipelineState);
@@ -310,7 +310,7 @@ void SetGraphicsPipelineState(FRHICommandList& RHICmdList, const FGraphicsPipeli
 		check(Result >= 1);
 #endif
 		check(IsInRenderingThread() || IsInParallelRenderingThread());
-		RHICmdList.SetGraphicsPipelineState(PipelineState);
+		RHICmdList.SetGraphicsPipelineState(PipelineState, Initializer.BoundShaderState);
 	}
 }
 
@@ -1141,20 +1141,6 @@ FGraphicsPipelineState* PipelineStateCache::GetAndOrCreateGraphicsPipelineState(
 {
 	LLM_SCOPE(ELLMTag::PSO);
 
-	// Workaround until we have a better way for storing shaders in PSO cache 
-	{
-		FGraphicsPipelineStateInitializer& HashableInitializer = const_cast<FGraphicsPipelineStateInitializer&>(OriginalInitializer);
-		HashableInitializer.VertexShaderHash = HashableInitializer.BoundShaderState.VertexShaderRHI ? HashableInitializer.BoundShaderState.VertexShaderRHI->GetHash() : FSHAHash();
-#if PLATFORM_SUPPORTS_TESSELLATION_SHADERS
-		HashableInitializer.HullShaderHash = HashableInitializer.BoundShaderState.HullShaderRHI ? HashableInitializer.BoundShaderState.HullShaderRHI->GetHash() : FSHAHash();
-		HashableInitializer.DomainShaderHash = HashableInitializer.BoundShaderState.DomainShaderRHI ? HashableInitializer.BoundShaderState.DomainShaderRHI->GetHash() : FSHAHash();
-#endif
-		HashableInitializer.PixelShaderHash = HashableInitializer.BoundShaderState.PixelShaderRHI ? HashableInitializer.BoundShaderState.PixelShaderRHI->GetHash() : FSHAHash();
-#if PLATFORM_SUPPORTS_GEOMETRY_SHADERS
-		HashableInitializer.GeometryShaderHash = HashableInitializer.BoundShaderState.GeometryShaderRHI ? HashableInitializer.BoundShaderState.GeometryShaderRHI->GetHash() : FSHAHash();
-#endif
-	}
-	
 	FGraphicsPipelineStateInitializer NewInitializer;
 	const FGraphicsPipelineStateInitializer* Initializer = &OriginalInitializer;
 
diff --git a/Engine/Source/Runtime/RHI/Private/RHI.cpp b/Engine/Source/Runtime/RHI/Private/RHI.cpp
--- a/Engine/Source/Runtime/RHI/Private/RHI.cpp
+++ b/Engine/Source/Runtime/RHI/Private/RHI.cpp
@@ -37,6 +37,9 @@ DEFINE_STAT(STAT_VertexBufferMemory);
 DEFINE_STAT(STAT_StructuredBufferMemory);
 DEFINE_STAT(STAT_PixelBufferMemory);
 
+IMPLEMENT_TYPE_LAYOUT(FRHIUniformBufferLayout);
+IMPLEMENT_TYPE_LAYOUT(FRHIUniformBufferLayout::FResourceParameter);
+
 static FAutoConsoleVariable CVarUseVulkanRealUBs(
 	TEXT("r.Vulkan.UseRealUBs"),
 	1,
@@ -846,7 +849,7 @@ RHI_API const TCHAR* RHIVendorIdToString(EGpuVendorId VendorId)
 	return TEXT("Unknown");
 }
 
-RHI_API uint32 RHIGetShaderLanguageVersion(const EShaderPlatform Platform)
+RHI_API uint32 RHIGetShaderLanguageVersion(const FStaticShaderPlatform Platform)
 {
 	uint32 Version = 0;
 	if (IsMetalPlatform(Platform))
@@ -885,7 +888,7 @@ RHI_API uint32 RHIGetShaderLanguageVersion(const EShaderPlatform Platform)
 	return Version;
 }
 
-RHI_API bool RHISupportsTessellation(const EShaderPlatform Platform)
+RHI_API bool RHISupportsTessellation(const FStaticShaderPlatform Platform)
 {
 	if (IsFeatureLevelSupported(Platform, ERHIFeatureLevel::SM5))
 	{
@@ -894,7 +897,7 @@ RHI_API bool RHISupportsTessellation(const EShaderPlatform Platform)
 	return false;
 }
 
-RHI_API bool RHISupportsPixelShaderUAVs(const EShaderPlatform Platform)
+RHI_API bool RHISupportsPixelShaderUAVs(const FStaticShaderPlatform Platform)
 {
 	if (IsFeatureLevelSupported(Platform, ERHIFeatureLevel::SM5))
 	{
@@ -903,10 +906,10 @@ RHI_API bool RHISupportsPixelShaderUAVs(const EShaderPlatform Platform)
 	return false;
 }
 
-RHI_API bool RHISupportsIndexBufferUAVs(const EShaderPlatform Platform)
+RHI_API bool RHISupportsIndexBufferUAVs(const FStaticShaderPlatform Platform)
 {
 	return Platform == SP_PCD3D_SM5 || IsVulkanPlatform(Platform) || IsMetalSM5Platform(Platform) || Platform == SP_XBOXONE_D3D12 || Platform == SP_PS4 
-		|| FDataDrivenShaderPlatformInfo::GetInfo(Platform).bSupportsIndexBufferUAVs;
+		|| FDataDrivenShaderPlatformInfo::GetSupportsIndexBufferUAVs(Platform);
 }
 
 
@@ -1204,14 +1207,14 @@ void LexFromString(EShaderPlatform& Value, const TCHAR* String)
 }
 
 
-FName LANGUAGE_D3D("D3D");
-FName LANGUAGE_Metal("Metal");
-FName LANGUAGE_OpenGL("OpenGL");
-FName LANGUAGE_Vulkan("Vulkan");
-FName LANGUAGE_Sony("Sony");
-FName LANGUAGE_Nintendo("Nintendo");
+const FName LANGUAGE_D3D("D3D");
+const FName LANGUAGE_Metal("Metal");
+const FName LANGUAGE_OpenGL("OpenGL");
+const FName LANGUAGE_Vulkan("Vulkan");
+const FName LANGUAGE_Sony("Sony");
+const FName LANGUAGE_Nintendo("Nintendo");
 
-RHI_API FDataDrivenShaderPlatformInfo FDataDrivenShaderPlatformInfo::Infos[SP_NumPlatforms];
+RHI_API FGenericDataDrivenShaderPlatformInfo FGenericDataDrivenShaderPlatformInfo::Infos[SP_NumPlatforms];
 
 // Gets a string from a section, or empty string if it didn't exist
 FString GetSectionString(const FConfigSection& Section, FName Key)
@@ -1225,7 +1228,7 @@ bool GetSectionBool(const FConfigSection& Section, FName Key)
 	return FCString::ToBool(*GetSectionString(Section, Key));
 }
 
-inline void ParseDataDrivenShaderInfo(const FConfigSection& Section, FDataDrivenShaderPlatformInfo& Info)
+void FGenericDataDrivenShaderPlatformInfo::ParseDataDrivenShaderInfo(const FConfigSection& Section, FGenericDataDrivenShaderPlatformInfo& Info)
 {
 	Info.Language = *GetSectionString(Section, "Language");
 	GetFeatureLevelFromName(*GetSectionString(Section, "MaxFeatureLevel"), Info.MaxFeatureLevel);
@@ -1253,12 +1256,13 @@ inline void ParseDataDrivenShaderInfo(const FConfigSection& Section, FDataDriven
 	Info.bSupportsRayTracingMissShaderBindings = GetSectionBool(Section, "bSupportsRayTracingMissShaderBindings");
 	Info.bSupportsGPUSkinCache = GetSectionBool(Section, "bSupportsGPUSkinCache");
 	Info.bSupportsByteBufferComputeShaders = GetSectionBool(Section, "bSupportsByteBufferComputeShaders");
+	Info.bSupportsGPUScene = GetSectionBool(Section, "bSupportsGPUScene");
 
 	Info.bTargetsTiledGPU = GetSectionBool(Section, "bTargetsTiledGPU");
 	Info.bNeedsOfflineCompiler = GetSectionBool(Section, "bNeedsOfflineCompiler");
 }
 
-void FDataDrivenShaderPlatformInfo::Initialize()
+void FGenericDataDrivenShaderPlatformInfo::Initialize()
 {
 	// look for the standard DataDriven ini files
 	int32 NumDDInfoFiles = FDataDrivenPlatformInfoRegistry::GetNumDataDrivenIniFiles();
@@ -1287,6 +1291,7 @@ void FDataDrivenShaderPlatformInfo::Initialize()
 				
 				// at this point, we can start pulling information out
 				ParseDataDrivenShaderInfo(Section.Value, Infos[ShaderPlatform]);	
+				Infos[ShaderPlatform].bContainsValidPlatformInfo = true;
 			}
 		}
 	}
diff --git a/Engine/Source/Runtime/RHI/Private/RHICommandList.cpp b/Engine/Source/Runtime/RHI/Private/RHICommandList.cpp
--- a/Engine/Source/Runtime/RHI/Private/RHICommandList.cpp
+++ b/Engine/Source/Runtime/RHI/Private/RHICommandList.cpp
@@ -848,6 +848,7 @@ FRHICommandListBase::FRHICommandListBase(FRHIGPUMask InGPUMask)
 	, bAsyncPSOCompileAllowed(true)
 	, GPUMask(InGPUMask)
 	, InitialGPUMask(InGPUMask)
+	, BoundComputeShaderRHI(nullptr)
 {
 	GRHICommandList.OutstandingCmdListCount.Increment();
 	Reset();
@@ -1993,7 +1994,7 @@ void* FDynamicRHI::RHILockVertexBuffer(class FRHICommandListImmediate& RHICmdLis
 	}
 
 	check(Result);
-	return Result;	
+	return Result;
 }
 
 void FDynamicRHI::RHIUnlockVertexBuffer(class FRHICommandListImmediate& RHICmdList, FRHIVertexBuffer* VertexBuffer)
@@ -2031,7 +2032,7 @@ void FDynamicRHI::RHIUnlockVertexBuffer(class FRHICommandListImmediate& RHICmdLi
 		{
 			QUICK_SCOPE_CYCLE_COUNTER(STAT_RHIMETHOD_UnlockVertexBuffer_FlushForMem);
 			// we could be loading a level or something, lets get this stuff going
-			RHICmdList.ImmediateFlush(EImmediateFlushType::DispatchToRHIThread);
+			RHICmdList.ImmediateFlush(EImmediateFlushType::DispatchToRHIThread); 
 			GLockTracker.TotalMemoryOutstanding = 0;
 		}
 	}
@@ -2111,7 +2112,7 @@ void FDynamicRHI::RHIUnlockIndexBuffer(class FRHICommandListImmediate& RHICmdLis
 		{
 			QUICK_SCOPE_CYCLE_COUNTER(STAT_RHIMETHOD_UnlockIndexBuffer_FlushForMem);
 			// we could be loading a level or something, lets get this stuff going
-			RHICmdList.ImmediateFlush(EImmediateFlushType::DispatchToRHIThread);
+			RHICmdList.ImmediateFlush(EImmediateFlushType::DispatchToRHIThread); 
 			GLockTracker.TotalMemoryOutstanding = 0;
 		}
 	}
@@ -2191,7 +2192,7 @@ void FDynamicRHI::RHIUnlockStructuredBuffer(class FRHICommandListImmediate& RHIC
 		{
 			QUICK_SCOPE_CYCLE_COUNTER(STAT_RHIMETHOD_UnlockStructuredBuffer_FlushForMem);
 			// we could be loading a level or something, lets get this stuff going
-			RHICmdList.ImmediateFlush(EImmediateFlushType::DispatchToRHIThread);
+			RHICmdList.ImmediateFlush(EImmediateFlushType::DispatchToRHIThread); 
 			GLockTracker.TotalMemoryOutstanding = 0;
 		}
 	}
@@ -2273,88 +2274,46 @@ ETextureReallocationStatus FDynamicRHI::CancelAsyncReallocateTexture2D_RenderThr
 	return GDynamicRHI->RHICancelAsyncReallocateTexture2D(Texture2D, bBlockUntilCompleted);
 }
 
-FVertexShaderRHIRef FDynamicRHI::CreateVertexShader_RenderThread(class FRHICommandListImmediate& RHICmdList, const TArray<uint8>& Code)
-{
-	CSV_SCOPED_TIMING_STAT(RHITStalls, CreateVertexShader_RenderThread);
-	FScopedRHIThreadStaller StallRHIThread(RHICmdList);
-	return GDynamicRHI->RHICreateVertexShader(Code);
-}
-
-FVertexShaderRHIRef FDynamicRHI::CreateVertexShader_RenderThread(class FRHICommandListImmediate& RHICmdList, FRHIShaderLibrary* Library, FSHAHash Hash)
+FVertexShaderRHIRef FDynamicRHI::CreateVertexShader_RenderThread(class FRHICommandListImmediate& RHICmdList, TArrayView<const uint8> Code, const FSHAHash& Hash)
 {
 	CSV_SCOPED_TIMING_STAT(RHITStalls, CreateVertexShader_RenderThread);
 	FScopedRHIThreadStaller StallRHIThread(RHICmdList);
-	return GDynamicRHI->RHICreateVertexShader(Library, Hash);
+	return GDynamicRHI->RHICreateVertexShader(Code, Hash);
 }
 
-FPixelShaderRHIRef FDynamicRHI::CreatePixelShader_RenderThread(class FRHICommandListImmediate& RHICmdList, const TArray<uint8>& Code)
+FPixelShaderRHIRef FDynamicRHI::CreatePixelShader_RenderThread(class FRHICommandListImmediate& RHICmdList, TArrayView<const uint8> Code, const FSHAHash& Hash)
 {
 	CSV_SCOPED_TIMING_STAT(RHITStalls, CreatePixelShader_RenderThread);
 	FScopedRHIThreadStaller StallRHIThread(RHICmdList);
-	return GDynamicRHI->RHICreatePixelShader(Code);
-}
-
-FPixelShaderRHIRef FDynamicRHI::CreatePixelShader_RenderThread(class FRHICommandListImmediate& RHICmdList, FRHIShaderLibrary* Library, FSHAHash Hash)
-{
-	CSV_SCOPED_TIMING_STAT(RHITStalls, CreatePixelShader_RenderThread);
-	FScopedRHIThreadStaller StallRHIThread(RHICmdList);
-	return GDynamicRHI->RHICreatePixelShader(Library, Hash);
-}
-
-FGeometryShaderRHIRef FDynamicRHI::CreateGeometryShader_RenderThread(class FRHICommandListImmediate& RHICmdList, const TArray<uint8>& Code)
-{
-	CSV_SCOPED_TIMING_STAT(RHITStalls, CreateGeometryShader_RenderThread);
-	FScopedRHIThreadStaller StallRHIThread(RHICmdList);
-	return GDynamicRHI->RHICreateGeometryShader(Code);
+	return GDynamicRHI->RHICreatePixelShader(Code, Hash);
 }
 
-FGeometryShaderRHIRef FDynamicRHI::CreateGeometryShader_RenderThread(class FRHICommandListImmediate& RHICmdList, FRHIShaderLibrary* Library, FSHAHash Hash)
+FGeometryShaderRHIRef FDynamicRHI::CreateGeometryShader_RenderThread(class FRHICommandListImmediate& RHICmdList, TArrayView<const uint8> Code, const FSHAHash& Hash)
 {
 	CSV_SCOPED_TIMING_STAT(RHITStalls, CreateGeometryShader_RenderThread);
 	FScopedRHIThreadStaller StallRHIThread(RHICmdList);
-	return GDynamicRHI->RHICreateGeometryShader(Library, Hash);
+	return GDynamicRHI->RHICreateGeometryShader(Code, Hash);
 }
 
-FComputeShaderRHIRef FDynamicRHI::CreateComputeShader_RenderThread(class FRHICommandListImmediate& RHICmdList, const TArray<uint8>& Code)
+FComputeShaderRHIRef FDynamicRHI::CreateComputeShader_RenderThread(class FRHICommandListImmediate& RHICmdList, TArrayView<const uint8> Code, const FSHAHash& Hash)
 {
 	CSV_SCOPED_TIMING_STAT(RHITStalls, CreateGeometryShader_RenderThread);
 	FScopedRHIThreadStaller StallRHIThread(RHICmdList);
-	return GDynamicRHI->RHICreateComputeShader(Code);
-}
-
-FComputeShaderRHIRef FDynamicRHI::CreateComputeShader_RenderThread(class FRHICommandListImmediate& RHICmdList, FRHIShaderLibrary* Library, FSHAHash Hash)
-{
-	CSV_SCOPED_TIMING_STAT(RHITStalls, CreateComputeShader_RenderThread);
-	FScopedRHIThreadStaller StallRHIThread(RHICmdList);
-	return GDynamicRHI->RHICreateComputeShader(Library, Hash);
+	return GDynamicRHI->RHICreateComputeShader(Code, Hash);
 }
 
-FHullShaderRHIRef FDynamicRHI::CreateHullShader_RenderThread(class FRHICommandListImmediate& RHICmdList, const TArray<uint8>& Code)
+FHullShaderRHIRef FDynamicRHI::CreateHullShader_RenderThread(class FRHICommandListImmediate& RHICmdList, TArrayView<const uint8> Code, const FSHAHash& Hash)
 {
 	CSV_SCOPED_TIMING_STAT(RHITStalls, CreateHullShader_RenderThread);
 	FScopedRHIThreadStaller StallRHIThread(RHICmdList);
-	return GDynamicRHI->RHICreateHullShader(Code);
-}
-
-FHullShaderRHIRef FDynamicRHI::CreateHullShader_RenderThread(class FRHICommandListImmediate& RHICmdList, FRHIShaderLibrary* Library, FSHAHash Hash)
-{
-	CSV_SCOPED_TIMING_STAT(RHITStalls, CreateHullShader_RenderThread);
-	FScopedRHIThreadStaller StallRHIThread(RHICmdList);
-	return GDynamicRHI->RHICreateHullShader(Library, Hash);
-}
-
-FDomainShaderRHIRef FDynamicRHI::CreateDomainShader_RenderThread(class FRHICommandListImmediate& RHICmdList, const TArray<uint8>& Code)
-{
-	CSV_SCOPED_TIMING_STAT(RHITStalls, CreateDomainShader_RenderThread);
-	FScopedRHIThreadStaller StallRHIThread(RHICmdList);
-	return GDynamicRHI->RHICreateDomainShader(Code);
+	return GDynamicRHI->RHICreateHullShader(Code, Hash);
 }
 
-FDomainShaderRHIRef FDynamicRHI::CreateDomainShader_RenderThread(class FRHICommandListImmediate& RHICmdList, FRHIShaderLibrary* Library, FSHAHash Hash)
+FDomainShaderRHIRef FDynamicRHI::CreateDomainShader_RenderThread(class FRHICommandListImmediate& RHICmdList, TArrayView<const uint8> Code, const FSHAHash& Hash)
 {
 	CSV_SCOPED_TIMING_STAT(RHITStalls, CreateDomainShader_RenderThread);
 	FScopedRHIThreadStaller StallRHIThread(RHICmdList);
-	return GDynamicRHI->RHICreateDomainShader(Library, Hash);
+	return GDynamicRHI->RHICreateDomainShader(Code, Hash);
 }
 
 void FDynamicRHI::UpdateTexture2D_RenderThread(class FRHICommandListImmediate& RHICmdList, FRHITexture2D* Texture, uint32 MipIndex, const struct FUpdateTextureRegion2D& UpdateRegion, uint32 SourcePitch, const uint8* SourceData)
diff --git a/Engine/Source/Runtime/RHI/Public/DynamicRHI.h b/Engine/Source/Runtime/RHI/Public/DynamicRHI.h
--- a/Engine/Source/Runtime/RHI/Public/DynamicRHI.h
+++ b/Engine/Source/Runtime/RHI/Public/DynamicRHI.h
@@ -9,6 +9,7 @@ DynamicRHI.h: Dynamically bound Render Hardware Interface definitions.
 #include "CoreTypes.h"
 #include "RHIContext.h"
 #include "MultiGPU.h"
+#include "Serialization/MemoryLayout.h"
 
 class FBlendStateInitializerRHI;
 class FGraphicsPipelineStateInitializer;
@@ -116,53 +117,58 @@ enum ERayTracingGeometryType
 	// Index buffers can't be used with procedural geometry.
 	RTGT_Procedural,
 };
+DECLARE_INTRINSIC_TYPE_LAYOUT(ERayTracingGeometryType);
 
 struct FRayTracingGeometrySegment
 {
-	FVertexBufferRHIRef VertexBuffer = nullptr;
-	EVertexElementType VertexBufferElementType = VET_Float3;
+	DECLARE_TYPE_LAYOUT(FRayTracingGeometrySegment, NonVirtual);
+public:
+	LAYOUT_FIELD_INITIALIZED(FVertexBufferRHIRef, VertexBuffer, nullptr);
+	LAYOUT_FIELD_INITIALIZED(EVertexElementType, VertexBufferElementType, VET_Float3);
 
 	// Offset in bytes from the base address of the vertex buffer.
-	uint32 VertexBufferOffset = 0;
+	LAYOUT_FIELD_INITIALIZED(uint32, VertexBufferOffset, 0);
 
 	// Number of bytes between elements of the vertex buffer (sizeof VET_Float3 by default).
 	// Must be equal or greater than the size of the position vector.
-	uint32 VertexBufferStride = 12;
+	LAYOUT_FIELD_INITIALIZED(uint32, VertexBufferStride, 12);
 
 	// Primitive range for this segment.
-	uint32 FirstPrimitive = 0;
-	uint32 NumPrimitives = 0;
+	LAYOUT_FIELD_INITIALIZED(uint32, FirstPrimitive, 0);
+	LAYOUT_FIELD_INITIALIZED(uint32, NumPrimitives, 0);
 
 	// Indicates whether any-hit shader could be invoked when hitting this geometry segment.
 	// Setting this to `false` turns off any-hit shaders, making the section "opaque" and improving ray tracing performance.
-	bool bForceOpaque = false;
+	LAYOUT_FIELD_INITIALIZED(bool, bForceOpaque, false);
 
 	// Any-hit shader may be invoked multiple times for the same primitive during ray traversal.
 	// Setting this to `false` guarantees that only a single instance of any-hit shader will run per primitive, at some performance cost.
-	bool bAllowDuplicateAnyHitShaderInvocation = true;
+	LAYOUT_FIELD_INITIALIZED(bool, bAllowDuplicateAnyHitShaderInvocation, true);
 
 	// Indicates whether this section is enabled and should be taken into account during acceleration structure creation
-	bool bEnabled = true;
+	LAYOUT_FIELD_INITIALIZED(bool, bEnabled, true);
 };
 
 struct FRayTracingGeometryInitializer
 {
-	FIndexBufferRHIRef IndexBuffer = nullptr;
+	DECLARE_EXPORTED_TYPE_LAYOUT(FRayTracingGeometryInitializer, RHI_API, NonVirtual);
+public:
+	LAYOUT_FIELD_INITIALIZED(FIndexBufferRHIRef, IndexBuffer, nullptr);
 
 	// Offset in bytes from the base address of the index buffer.
-	uint32 IndexBufferOffset = 0;
-
-	ERayTracingGeometryType GeometryType = RTGT_Triangles;
+	LAYOUT_FIELD_INITIALIZED(uint32, IndexBufferOffset, 0);
+	
+	LAYOUT_FIELD_INITIALIZED(ERayTracingGeometryType, GeometryType, RTGT_Triangles);
 
 	// Total number of primitives in all segments of the geometry. Only used for validation.
-	uint32 TotalPrimitiveCount = 0;
+	LAYOUT_FIELD_INITIALIZED(uint32, TotalPrimitiveCount, 0);
 
 	// Partitions of geometry to allow different shader and resource bindings.
 	// All ray tracing geometries must have at least one segment.
-	TArray<FRayTracingGeometrySegment> Segments;
+	LAYOUT_FIELD(TMemoryImageArray<FRayTracingGeometrySegment>, Segments);
 
-	bool bFastBuild = false;
-	bool bAllowUpdate = false;
+	LAYOUT_FIELD_INITIALIZED(bool, bFastBuild, false);
+	LAYOUT_FIELD_INITIALIZED(bool, bAllowUpdate, false);
 };
 
 enum ERayTracingSceneLifetime
@@ -350,61 +356,25 @@ public:
 	virtual FVertexDeclarationRHIRef RHICreateVertexDeclaration(const FVertexDeclarationElementList& Elements) = 0;
 
 	// FlushType: Wait RHI Thread
-	virtual FPixelShaderRHIRef RHICreatePixelShader(const TArray<uint8>& Code) = 0;
-
-	// FlushType: Wait RHI Thread
-	virtual FPixelShaderRHIRef RHICreatePixelShader(FRHIShaderLibrary* Library, FSHAHash Hash)
-	{
-		return nullptr;
-	}
-
-	// FlushType: Wait RHI Thread
-	virtual FVertexShaderRHIRef RHICreateVertexShader(const TArray<uint8>& Code) = 0;
+	virtual FPixelShaderRHIRef RHICreatePixelShader(TArrayView<const uint8> Code, const FSHAHash& Hash) = 0;
 
 	// FlushType: Wait RHI Thread
-	virtual FVertexShaderRHIRef RHICreateVertexShader(FRHIShaderLibrary* Library, FSHAHash Hash)
-	{
-		return nullptr;
-	}
+	virtual FVertexShaderRHIRef RHICreateVertexShader(TArrayView<const uint8> Code, const FSHAHash& Hash) = 0;
 
 	// FlushType: Wait RHI Thread
-	virtual FHullShaderRHIRef RHICreateHullShader(const TArray<uint8>& Code) = 0;
+	virtual FHullShaderRHIRef RHICreateHullShader(TArrayView<const uint8> Code, const FSHAHash& Hash) = 0;
 
 	// FlushType: Wait RHI Thread
-	virtual FHullShaderRHIRef RHICreateHullShader(FRHIShaderLibrary* Library, FSHAHash Hash)
-	{
-		return nullptr;
-	}
+	virtual FDomainShaderRHIRef RHICreateDomainShader(TArrayView<const uint8> Code, const FSHAHash& Hash) = 0;
 
 	// FlushType: Wait RHI Thread
-	virtual FDomainShaderRHIRef RHICreateDomainShader(const TArray<uint8>& Code) = 0;
-
-	// FlushType: Wait RHI Thread
-	virtual FDomainShaderRHIRef RHICreateDomainShader(FRHIShaderLibrary* Library, FSHAHash Hash)
-	{
-		return nullptr;
-	}
-
-	// FlushType: Wait RHI Thread
-	virtual FGeometryShaderRHIRef RHICreateGeometryShader(const TArray<uint8>& Code) = 0;
-
-	// FlushType: Wait RHI Thread
-	virtual FGeometryShaderRHIRef RHICreateGeometryShader(FRHIShaderLibrary* Library, FSHAHash Hash)
-	{
-		return nullptr;
-	}
+	virtual FGeometryShaderRHIRef RHICreateGeometryShader(TArrayView<const uint8> Code, const FSHAHash& Hash) = 0;
 
 	// Some RHIs can have pending messages/logs for error tracking, or debug modes
 	virtual void FlushPendingLogs() {}
 
 	// FlushType: Wait RHI Thread
-	virtual FComputeShaderRHIRef RHICreateComputeShader(const TArray<uint8>& Code) = 0;
-
-	// FlushType: Wait RHI Thread
-	virtual FComputeShaderRHIRef RHICreateComputeShader(FRHIShaderLibrary* Library, FSHAHash Hash)
-	{
-		return nullptr;
-	}
+	virtual FComputeShaderRHIRef RHICreateComputeShader(TArrayView<const uint8> Code, const FSHAHash& Hash) = 0;
 
 	/**
 	 * Attempts to open a shader library for the given shader platform & name within the provided directory.
@@ -1307,18 +1277,12 @@ public:
 	virtual ETextureReallocationStatus FinalizeAsyncReallocateTexture2D_RenderThread(class FRHICommandListImmediate& RHICmdList, FRHITexture2D* Texture2D, bool bBlockUntilCompleted);
 	virtual ETextureReallocationStatus CancelAsyncReallocateTexture2D_RenderThread(class FRHICommandListImmediate& RHICmdList, FRHITexture2D* Texture2D, bool bBlockUntilCompleted);
 	virtual FIndexBufferRHIRef CreateIndexBuffer_RenderThread(class FRHICommandListImmediate& RHICmdList, uint32 Stride, uint32 Size, uint32 InUsage, FRHIResourceCreateInfo& CreateInfo);
-	virtual FVertexShaderRHIRef CreateVertexShader_RenderThread(class FRHICommandListImmediate& RHICmdList, const TArray<uint8>& Code);
-	virtual FVertexShaderRHIRef CreateVertexShader_RenderThread(class FRHICommandListImmediate& RHICmdList, FRHIShaderLibrary* Library, FSHAHash Hash);
-	virtual FPixelShaderRHIRef CreatePixelShader_RenderThread(class FRHICommandListImmediate& RHICmdList, const TArray<uint8>& Code);
-	virtual FPixelShaderRHIRef CreatePixelShader_RenderThread(class FRHICommandListImmediate& RHICmdList, FRHIShaderLibrary* Library, FSHAHash Hash);
-	virtual FGeometryShaderRHIRef CreateGeometryShader_RenderThread(class FRHICommandListImmediate& RHICmdList, const TArray<uint8>& Code);
-	virtual FGeometryShaderRHIRef CreateGeometryShader_RenderThread(class FRHICommandListImmediate& RHICmdList, FRHIShaderLibrary* Library, FSHAHash Hash);
-	virtual FComputeShaderRHIRef CreateComputeShader_RenderThread(class FRHICommandListImmediate& RHICmdList, const TArray<uint8>& Code);
-	virtual FComputeShaderRHIRef CreateComputeShader_RenderThread(class FRHICommandListImmediate& RHICmdList, FRHIShaderLibrary* Library, FSHAHash Hash);
-	virtual FHullShaderRHIRef CreateHullShader_RenderThread(class FRHICommandListImmediate& RHICmdList, const TArray<uint8>& Code);
-	virtual FHullShaderRHIRef CreateHullShader_RenderThread(class FRHICommandListImmediate& RHICmdList, FRHIShaderLibrary* Library, FSHAHash Hash);
-	virtual FDomainShaderRHIRef CreateDomainShader_RenderThread(class FRHICommandListImmediate& RHICmdList, const TArray<uint8>& Code);
-	virtual FDomainShaderRHIRef CreateDomainShader_RenderThread(class FRHICommandListImmediate& RHICmdList, FRHIShaderLibrary* Library, FSHAHash Hash);
+	virtual FVertexShaderRHIRef CreateVertexShader_RenderThread(class FRHICommandListImmediate& RHICmdList, TArrayView<const uint8> Code, const FSHAHash& Hash);
+	virtual FPixelShaderRHIRef CreatePixelShader_RenderThread(class FRHICommandListImmediate& RHICmdList, TArrayView<const uint8> Code, const FSHAHash& Hash);
+	virtual FGeometryShaderRHIRef CreateGeometryShader_RenderThread(class FRHICommandListImmediate& RHICmdList, TArrayView<const uint8> Code, const FSHAHash& Hash);
+	virtual FComputeShaderRHIRef CreateComputeShader_RenderThread(class FRHICommandListImmediate& RHICmdList, TArrayView<const uint8> Code, const FSHAHash& Hash);
+	virtual FHullShaderRHIRef CreateHullShader_RenderThread(class FRHICommandListImmediate& RHICmdList, TArrayView<const uint8> Code, const FSHAHash& Hash);
+	virtual FDomainShaderRHIRef CreateDomainShader_RenderThread(class FRHICommandListImmediate& RHICmdList, TArrayView<const uint8> Code, const FSHAHash& Hash);
 	virtual void* LockTexture2D_RenderThread(class FRHICommandListImmediate& RHICmdList, FRHITexture2D* Texture, uint32 MipIndex, EResourceLockMode LockMode, uint32& DestStride, bool bLockWithinMiptail, bool bNeedsDefaultRHIFlush = true);
 	virtual void UnlockTexture2D_RenderThread(class FRHICommandListImmediate& RHICmdList, FRHITexture2D* Texture, uint32 MipIndex, bool bLockWithinMiptail, bool bNeedsDefaultRHIFlush = true);
 	virtual void UpdateTexture2D_RenderThread(class FRHICommandListImmediate& RHICmdList, FRHITexture2D* Texture, uint32 MipIndex, const struct FUpdateTextureRegion2D& UpdateRegion, uint32 SourcePitch, const uint8* SourceData);
@@ -1450,7 +1414,7 @@ public:
 		return nullptr;
 	}
 
-	virtual FRayTracingShaderRHIRef RHICreateRayTracingShader(const TArray<uint8>& Code, EShaderFrequency ShaderFrequency)
+	virtual FRayTracingShaderRHIRef RHICreateRayTracingShader(TArrayView<const uint8> Code, const FSHAHash& Hash, EShaderFrequency ShaderFrequency)
 	{
 		checkNoEntry();
 		return nullptr;
@@ -1699,9 +1663,9 @@ FORCEINLINE FRayTracingSceneRHIRef RHICreateRayTracingScene(const FRayTracingSce
 	return GDynamicRHI->RHICreateRayTracingScene(Initializer);
 }
 
-FORCEINLINE FRayTracingShaderRHIRef RHICreateRayTracingShader(const TArray<uint8>& Code, EShaderFrequency ShaderFrequency)
+FORCEINLINE FRayTracingShaderRHIRef RHICreateRayTracingShader(TArrayView<const uint8> Code, const FSHAHash& Hash, EShaderFrequency ShaderFrequency)
 {
-	return GDynamicRHI->RHICreateRayTracingShader(Code, ShaderFrequency);
+	return GDynamicRHI->RHICreateRayTracingShader(Code, Hash, ShaderFrequency);
 }
 
 #endif // RHI_RAYTRACING
diff --git a/Engine/Source/Runtime/RHI/Public/RHI.h b/Engine/Source/Runtime/RHI/Public/RHI.h
--- a/Engine/Source/Runtime/RHI/Public/RHI.h
+++ b/Engine/Source/Runtime/RHI/Public/RHI.h
@@ -115,16 +115,16 @@ RHI_API const TCHAR* RHIVendorIdToString();
 RHI_API const TCHAR* RHIVendorIdToString(EGpuVendorId VendorId);
 
 // helper to return the shader language version for the given shader platform.
-RHI_API uint32 RHIGetShaderLanguageVersion(const EShaderPlatform Platform);
+RHI_API uint32 RHIGetShaderLanguageVersion(const FStaticShaderPlatform Platform);
 
 // helper to check that the shader platform supports tessellation.
-RHI_API bool RHISupportsTessellation(const EShaderPlatform Platform);
+RHI_API bool RHISupportsTessellation(const FStaticShaderPlatform Platform);
 
 // helper to check that the shader platform supports writing to UAVs from pixel shaders.
-RHI_API bool RHISupportsPixelShaderUAVs(const EShaderPlatform Platform);
+RHI_API bool RHISupportsPixelShaderUAVs(const FStaticShaderPlatform Platform);
 
 // helper to check that the shader platform supports creating a UAV off an index buffer.
-RHI_API bool RHISupportsIndexBufferUAVs(const EShaderPlatform Platform);
+RHI_API bool RHISupportsIndexBufferUAVs(const FStaticShaderPlatform Platform);
 
 // helper to check if a preview feature level has been requested.
 RHI_API bool RHIGetPreviewFeatureLevel(ERHIFeatureLevel::Type& PreviewFeatureLevelOUT);
@@ -132,42 +132,42 @@ RHI_API bool RHIGetPreviewFeatureLevel(ERHIFeatureLevel::Type& PreviewFeatureLev
 // helper to check if preferred EPixelFormat is supported, return one if it is not
 RHI_API EPixelFormat RHIPreferredPixelFormatHint(EPixelFormat PreferredPixelFormat);
 
-inline bool RHISupportsInstancedStereo(const EShaderPlatform Platform)
+inline bool RHISupportsInstancedStereo(const FStaticShaderPlatform Platform)
 {
 	// Only D3D SM5, PS4 and Metal SM5 supports Instanced Stereo
 	return Platform == EShaderPlatform::SP_PCD3D_SM5 || Platform == EShaderPlatform::SP_PS4 || Platform == EShaderPlatform::SP_METAL_SM5 || Platform == EShaderPlatform::SP_METAL_SM5_NOTESS
-		|| FDataDrivenShaderPlatformInfo::GetInfo(Platform).bSupportsInstancedStereo;
+		|| FDataDrivenShaderPlatformInfo::GetSupportsInstancedStereo(Platform);
 }
 
-inline bool RHISupportsMultiView(const EShaderPlatform Platform)
+inline bool RHISupportsMultiView(const FStaticShaderPlatform Platform)
 {
 	// Only PS4 and Metal SM5 from 10.13 onward supports Multi-View
 	return (Platform == EShaderPlatform::SP_PS4) || ((Platform == EShaderPlatform::SP_METAL_SM5 || Platform == SP_METAL_SM5_NOTESS))
-		|| FDataDrivenShaderPlatformInfo::GetInfo(Platform).bSupportsMultiView;
+		|| FDataDrivenShaderPlatformInfo::GetSupportsMultiView(Platform);
 }
 
-inline bool RHISupportsMSAA(EShaderPlatform Platform)
+inline bool RHISupportsMSAA(const FStaticShaderPlatform Platform)
 {
 	return 
 		(
 		// @todo optimise MSAA for XboxOne, currently uses significant eRAM.
 		Platform != SP_XBOXONE_D3D12)
 		// @todo platplug: Maybe this should become bDisallowMSAA to default of 0 is a better default (since now MSAA is opt-out more than opt-in) 
-		|| FDataDrivenShaderPlatformInfo::GetInfo(Platform).bSupportsMSAA;
+		|| FDataDrivenShaderPlatformInfo::GetSupportsMSAA(Platform);
 }
 
-inline bool RHISupportsBufferLoadTypeConversion(EShaderPlatform Platform)
+inline bool RHISupportsBufferLoadTypeConversion(const FStaticShaderPlatform Platform)
 {
 	return !IsMetalPlatform(Platform);
 }
 
 /** Whether the platform supports reading from volume textures (does not cover rendering to volume textures). */
-inline bool RHISupportsVolumeTextures(ERHIFeatureLevel::Type FeatureLevel)
+inline bool RHISupportsVolumeTextures(const FStaticFeatureLevel FeatureLevel)
 {
 	return FeatureLevel >= ERHIFeatureLevel::SM5;
 }
 
-inline bool RHISupportsVertexShaderLayer(const EShaderPlatform Platform)
+inline bool RHISupportsVertexShaderLayer(const FStaticShaderPlatform Platform)
 {
 	return IsFeatureLevelSupported(Platform, ERHIFeatureLevel::SM5) && IsMetalPlatform(Platform) && IsPCPlatform(Platform);
 }
@@ -175,24 +175,24 @@ inline bool RHISupportsVertexShaderLayer(const EShaderPlatform Platform)
 /** Return true if and only if the GPU support rendering to volume textures (2D Array, 3D) is guaranteed supported for a target platform.
 	if PipelineVolumeTextureLUTSupportGuaranteedAtRuntime is true then it is guaranteed that GSupportsVolumeTextureRendering is true at runtime.
 */
-inline bool RHIVolumeTextureRenderingSupportGuaranteed(const EShaderPlatform Platform)
+inline bool RHIVolumeTextureRenderingSupportGuaranteed(const FStaticShaderPlatform Platform)
 {
 	return IsFeatureLevelSupported(Platform, ERHIFeatureLevel::SM5)
 		&& (!IsMetalPlatform(Platform) || RHISupportsVertexShaderLayer(Platform)) // For Metal only shader platforms & versions that support vertex-shader-layer can render to volume textures - this is a compile/cook time check.
 		&& !IsOpenGLPlatform(Platform);		// Apparently, some OpenGL 3.3 cards support SM4 but can't render to volume textures
 }
 
-inline bool RHISupports4ComponentUAVReadWrite(EShaderPlatform Platform)
+inline bool RHISupports4ComponentUAVReadWrite(const FStaticShaderPlatform Platform)
 {
 	// Must match usf PLATFORM_SUPPORTS_4COMPONENT_UAV_READ_WRITE
 	// D3D11 does not support multi-component loads from a UAV: "error X3676: typed UAV loads are only allowed for single-component 32-bit element types"
 	return Platform == SP_XBOXONE_D3D12 || Platform == SP_PS4 || IsMetalPlatform(Platform) 
-		|| FDataDrivenShaderPlatformInfo::GetInfo(Platform).bSupports4ComponentUAVReadWrite;
+		|| FDataDrivenShaderPlatformInfo::GetSupports4ComponentUAVReadWrite(Platform);
 }
 
 /** Whether Manual Vertex Fetch is supported for the specified shader platform.
 	Shader Platform must not use the mobile renderer, and for Metal, the shader language must be at least 2. */
-inline bool RHISupportsManualVertexFetch(EShaderPlatform InShaderPlatform)
+inline bool RHISupportsManualVertexFetch(const FStaticShaderPlatform InShaderPlatform)
 {
 	return (!IsOpenGLPlatform(InShaderPlatform) || IsSwitchPlatform(InShaderPlatform)) && !IsMobilePlatform(InShaderPlatform);
 }
@@ -200,7 +200,7 @@ inline bool RHISupportsManualVertexFetch(EShaderPlatform InShaderPlatform)
 /** 
  * Returns true if SV_VertexID contains BaseVertexIndex passed to the draw call, false if shaders must manually construct an absolute VertexID.
  */
-inline bool RHISupportsAbsoluteVertexID(EShaderPlatform InShaderPlatform)
+inline bool RHISupportsAbsoluteVertexID(const EShaderPlatform InShaderPlatform)
 {
 	return IsVulkanPlatform(InShaderPlatform) || IsVulkanMobilePlatform(InShaderPlatform);
 }
@@ -208,27 +208,27 @@ inline bool RHISupportsAbsoluteVertexID(EShaderPlatform InShaderPlatform)
 /** Can this platform compile ray tracing shaders (regardless of project settings).
  *  To use at runtime, also check GRHISupportsRayTracing and r.RayTracing CVar (see IsRayTracingEnabled() helper).
  **/
-inline RHI_API bool RHISupportsRayTracingShaders(EShaderPlatform Platform)
+inline RHI_API bool RHISupportsRayTracingShaders(const FStaticShaderPlatform Platform)
 {
-	return FDataDrivenShaderPlatformInfo::GetInfo(Platform).bSupportsRayTracing;
+	return FDataDrivenShaderPlatformInfo::GetSupportsRayTracing(Platform);
 }
 
 /** Can this platform compile shaders that use shader model 6.0 wave intrinsics.
  *  To use such shaders at runtime, also check GRHISupportsWaveOperations.
  **/
-inline RHI_API bool RHISupportsWaveOperations(EShaderPlatform Platform)
+inline RHI_API bool RHISupportsWaveOperations(const FStaticShaderPlatform Platform)
 {
 	// Currently SM6 shaders are treated as an extension of SM5.
 	return Platform == SP_PCD3D_SM5;
 }
 
 /** True if the given shader platform supports a render target write mask */
-inline bool RHISupportsRenderTargetWriteMask(EShaderPlatform Platform)
+inline bool RHISupportsRenderTargetWriteMask(const FStaticShaderPlatform Platform)
 {
 	return
 		Platform == SP_PS4 ||
 		Platform == SP_XBOXONE_D3D12 ||
-		FDataDrivenShaderPlatformInfo::GetInfo(Platform).bSupportsRenderTargetWriteMask;
+		FDataDrivenShaderPlatformInfo::GetSupportsRenderTargetWriteMask(Platform);
 }
 
 // Wrapper for GRHI## global variables, allows values to be overridden for mobile preview modes.
@@ -600,7 +600,7 @@ extern RHI_API void GetFeatureLevelName(ERHIFeatureLevel::Type InFeatureLevel, F
 extern RHI_API EShaderPlatform GShaderPlatformForFeatureLevel[ERHIFeatureLevel::Num];
 
 /** Get the shader platform associated with the supplied feature level on this machine */
-inline EShaderPlatform GetFeatureLevelShaderPlatform(ERHIFeatureLevel::Type InFeatureLevel)
+inline EShaderPlatform GetFeatureLevelShaderPlatform(const FStaticFeatureLevel InFeatureLevel)
 {
 	return GShaderPlatformForFeatureLevel[InFeatureLevel];
 }
@@ -1728,16 +1728,6 @@ extern RHI_API void RHIPostInit(const TArray<uint32>& InPixelFormatByteWidth);
 extern RHI_API void RHIExit();
 
 
-// the following helper macros allow to safely convert shader types without much code clutter
-#define GETSAFERHISHADER_PIXEL(Shader) ((Shader) ? (Shader)->GetPixelShader() : nullptr)
-#define GETSAFERHISHADER_VERTEX(Shader) ((Shader) ? (Shader)->GetVertexShader() : nullptr)
-#define GETSAFERHISHADER_HULL(Shader) ((Shader) ? (Shader)->GetHullShader() : nullptr)
-#define GETSAFERHISHADER_DOMAIN(Shader) ((Shader) ? (Shader)->GetDomainShader() : nullptr)
-#define GETSAFERHISHADER_GEOMETRY(Shader) ((Shader) ? (Shader)->GetGeometryShader() : (FRHIGeometryShader*)FGeometryShaderRHIRef())
-#define GETSAFERHISHADER_COMPUTE(Shader) ((Shader) ? (Shader)->GetComputeShader() : nullptr)
-#define GETSAFERHISHADER_RAYTRACING(Shader) ((Shader) ? (Shader)->GetRayTracingShader() : (FRHIRayTracingShader*)FRayTracingShaderRHIRef())
-
-
 // Panic delegate is called when when a fatal condition is encountered within RHI function.
 DECLARE_DELEGATE_OneParam(FRHIPanicEvent, const FName&);
 extern RHI_API FRHIPanicEvent& RHIGetPanicDelegate();
diff --git a/Engine/Source/Runtime/RHI/Public/RHICommandList.h b/Engine/Source/Runtime/RHI/Public/RHICommandList.h
--- a/Engine/Source/Runtime/RHI/Public/RHICommandList.h
+++ b/Engine/Source/Runtime/RHI/Public/RHICommandList.h
@@ -505,6 +505,8 @@ public:
 		{
 			Other.CommandLink = &Other.Root;
 		}
+		Other.BoundShaderInput = BoundShaderInput;
+		Other.BoundComputeShaderRHI = BoundComputeShaderRHI;
 	}
 
 	void SetContext(IRHICommandContext* InContext)
@@ -598,6 +600,31 @@ protected:
 
 	} PSOContext;
 
+	FBoundShaderStateInput BoundShaderInput;
+	FRHIComputeShader* BoundComputeShaderRHI;
+
+	FORCEINLINE void ValidateBoundShader(FRHIVertexShader* ShaderRHI) { checkSlow(BoundShaderInput.VertexShaderRHI == ShaderRHI); }
+	FORCEINLINE void ValidateBoundShader(FRHIPixelShader* ShaderRHI) { checkSlow(BoundShaderInput.PixelShaderRHI == ShaderRHI); }
+	FORCEINLINE void ValidateBoundShader(FRHIGeometryShader* ShaderRHI) { checkSlow(BoundShaderInput.GeometryShaderRHI == ShaderRHI); }
+	FORCEINLINE void ValidateBoundShader(FRHIHullShader* ShaderRHI) { checkSlow(BoundShaderInput.HullShaderRHI == ShaderRHI); }
+	FORCEINLINE void ValidateBoundShader(FRHIDomainShader* ShaderRHI) { checkSlow(BoundShaderInput.DomainShaderRHI == ShaderRHI); }
+	FORCEINLINE void ValidateBoundShader(FRHIComputeShader* ShaderRHI) { checkSlow(BoundComputeShaderRHI == ShaderRHI); }
+
+	FORCEINLINE void ValidateBoundShader(FRHIGraphicsShader* ShaderRHI)
+	{
+#if DO_GUARD_SLOW
+		switch (ShaderRHI->GetFrequency())
+		{
+		case SF_Vertex: checkSlow(BoundShaderInput.VertexShaderRHI == ShaderRHI); break;
+		case SF_Hull: checkSlow(BoundShaderInput.HullShaderRHI == ShaderRHI); break;
+		case SF_Domain: checkSlow(BoundShaderInput.DomainShaderRHI == ShaderRHI); break;
+		case SF_Pixel: checkSlow(BoundShaderInput.PixelShaderRHI == ShaderRHI); break;
+		case SF_Geometry: checkSlow(BoundShaderInput.GeometryShaderRHI == ShaderRHI); break;
+		default: checkfSlow(false, TEXT("Unexpected graphics shader type %d"), ShaderRHI->GetFrequency());
+		}
+#endif // DO_GUARD_SLOW
+	}
+
 	void CacheActiveRenderTargets(
 		uint32 NewNumSimultaneousRenderTargets,
 		const FRHIRenderTargetView* NewRenderTargetsRHI,
@@ -2245,6 +2272,8 @@ public:
 	void* operator new(size_t Size);
 	void operator delete(void *RawMemory);
 
+	inline FRHIComputeShader* GetBoundComputeShader() const { return BoundComputeShaderRHI; }
+
 	FORCEINLINE_DEBUGGABLE void SetGlobalUniformBuffers(const FUniformBufferStaticBindings& UniformBuffers)
 	{
 		if (Bypass())
@@ -2257,6 +2286,7 @@ public:
 
 	FORCEINLINE_DEBUGGABLE void SetShaderUniformBuffer(FRHIComputeShader* Shader, uint32 BaseIndex, FRHIUniformBuffer* UniformBuffer)
 	{
+		ValidateBoundShader(Shader);
 		if (Bypass())
 		{
 			GetComputeContext().RHISetShaderUniformBuffer(Shader, BaseIndex, UniformBuffer);
@@ -2272,6 +2302,7 @@ public:
 
 	FORCEINLINE_DEBUGGABLE void SetShaderParameter(FRHIComputeShader* Shader, uint32 BufferIndex, uint32 BaseIndex, uint32 NumBytes, const void* NewValue)
 	{
+		ValidateBoundShader(Shader);
 		if (Bypass())
 		{
 			GetComputeContext().RHISetShaderParameter(Shader, BufferIndex, BaseIndex, NumBytes, NewValue);
@@ -2289,6 +2320,7 @@ public:
 
 	FORCEINLINE_DEBUGGABLE void SetShaderTexture(FRHIComputeShader* Shader, uint32 TextureIndex, FRHITexture* Texture)
 	{
+		ValidateBoundShader(Shader);
 		if (Bypass())
 		{
 			GetComputeContext().RHISetShaderTexture(Shader, TextureIndex, Texture);
@@ -2299,6 +2331,7 @@ public:
 
 	FORCEINLINE_DEBUGGABLE void SetShaderResourceViewParameter(FRHIComputeShader* Shader, uint32 SamplerIndex, FRHIShaderResourceView* SRV)
 	{
+		ValidateBoundShader(Shader);
 		if (Bypass())
 		{
 			GetComputeContext().RHISetShaderResourceViewParameter(Shader, SamplerIndex, SRV);
@@ -2326,6 +2359,7 @@ public:
 
 	FORCEINLINE_DEBUGGABLE void SetUAVParameter(FRHIComputeShader* Shader, uint32 UAVIndex, FRHIUnorderedAccessView* UAV)
 	{
+		ValidateBoundShader(Shader);
 		if (Bypass())
 		{
 			GetComputeContext().RHISetUAVParameter(Shader, UAVIndex, UAV);
@@ -2336,6 +2370,7 @@ public:
 
 	FORCEINLINE_DEBUGGABLE void SetUAVParameter(FRHIComputeShader* Shader, uint32 UAVIndex, FRHIUnorderedAccessView* UAV, uint32 InitialCount)
 	{
+		ValidateBoundShader(Shader);
 		if (Bypass())
 		{
 			GetComputeContext().RHISetUAVParameter(Shader, UAVIndex, UAV, InitialCount);
@@ -2346,6 +2381,7 @@ public:
 
 	FORCEINLINE_DEBUGGABLE void SetComputeShader(FRHIComputeShader* ComputeShader)
 	{
+		BoundComputeShaderRHI = ComputeShader;
 		ComputeShader->UpdateStats();
 		if (Bypass())
 		{
@@ -2355,8 +2391,9 @@ public:
 		ALLOC_COMMAND(FRHICommandSetComputeShader)(ComputeShader);
 	}
 
-	FORCEINLINE_DEBUGGABLE void SetComputePipelineState(FComputePipelineState* ComputePipelineState)
+	FORCEINLINE_DEBUGGABLE void SetComputePipelineState(FComputePipelineState* ComputePipelineState, FRHIComputeShader* ComputeShader)
 	{
+		BoundComputeShaderRHI = ComputeShader;
 		if (Bypass())
 		{
 			extern RHI_API FRHIComputePipelineState* ExecuteSetComputePipelineState(FComputePipelineState* ComputePipelineState);
@@ -2611,6 +2648,12 @@ public:
 	void* operator new(size_t Size);
 	void operator delete(void *RawMemory);
 	
+	inline FRHIVertexShader* GetBoundVertexShader() const { return BoundShaderInput.VertexShaderRHI; }
+	inline FRHIHullShader* GetBoundHullShader() const { return BoundShaderInput.HullShaderRHI; }
+	inline FRHIDomainShader* GetBoundDomainShader() const { return BoundShaderInput.DomainShaderRHI; }
+	inline FRHIPixelShader* GetBoundPixelShader() const { return BoundShaderInput.PixelShaderRHI; }
+	inline FRHIGeometryShader* GetBoundGeometryShader() const { return BoundShaderInput.GeometryShaderRHI; }
+
 	FORCEINLINE_DEBUGGABLE void BeginUpdateMultiFrameResource(FRHITexture* Texture)
 	{
 		//check(IsOutsideRenderPass());
@@ -2709,6 +2752,7 @@ public:
 	FORCEINLINE_DEBUGGABLE void SetLocalShaderUniformBuffer(TRHIShader* Shader, uint32 BaseIndex, const FLocalUniformBuffer& UniformBuffer)
 	{
 		//check(IsOutsideRenderPass());
+		ValidateBoundShader(Shader);
 		if (Bypass())
 		{
 			GetContext().RHISetShaderUniformBuffer(Shader, BaseIndex, UniformBuffer.BypassUniform);
@@ -2728,6 +2772,7 @@ public:
 	FORCEINLINE_DEBUGGABLE void SetShaderUniformBuffer(FRHIGraphicsShader* Shader, uint32 BaseIndex, FRHIUniformBuffer* UniformBuffer)
 	{
 		//check(IsOutsideRenderPass());
+		ValidateBoundShader(Shader);
 		if (Bypass())
 		{
 			GetContext().RHISetShaderUniformBuffer(Shader, BaseIndex, UniformBuffer);
@@ -2747,6 +2792,7 @@ public:
 	FORCEINLINE_DEBUGGABLE void SetShaderParameter(FRHIGraphicsShader* Shader, uint32 BufferIndex, uint32 BaseIndex, uint32 NumBytes, const void* NewValue)
 	{
 		//check(IsOutsideRenderPass());
+		ValidateBoundShader(Shader);
 		if (Bypass())
 		{
 			GetContext().RHISetShaderParameter(Shader, BufferIndex, BaseIndex, NumBytes, NewValue);
@@ -2768,6 +2814,7 @@ public:
 	FORCEINLINE_DEBUGGABLE void SetShaderTexture(FRHIGraphicsShader* Shader, uint32 TextureIndex, FRHITexture* Texture)
 	{
 		//check(IsOutsideRenderPass());
+		ValidateBoundShader(Shader);
 		if (Bypass())
 		{
 			GetContext().RHISetShaderTexture(Shader, TextureIndex, Texture);
@@ -2787,6 +2834,7 @@ public:
 	FORCEINLINE_DEBUGGABLE void SetShaderResourceViewParameter(FRHIGraphicsShader* Shader, uint32 SamplerIndex, FRHIShaderResourceView* SRV)
 	{
 		//check(IsOutsideRenderPass());
+		ValidateBoundShader(Shader);
 		if (Bypass())
 		{
 			GetContext().RHISetShaderResourceViewParameter(Shader, SamplerIndex, SRV);
@@ -2806,6 +2854,7 @@ public:
 	FORCEINLINE_DEBUGGABLE void SetShaderSampler(FRHIGraphicsShader* Shader, uint32 SamplerIndex, FRHISamplerState* State)
 	{
 		//check(IsOutsideRenderPass());
+		ValidateBoundShader(Shader);
 		
 		// Immutable samplers can't be set dynamically
 		check(!State->IsImmutable());
@@ -2832,6 +2881,7 @@ public:
 
 	FORCEINLINE_DEBUGGABLE void SetUAVParameter(FRHIPixelShader* Shader, uint32 UAVIndex, FRHIUnorderedAccessView* UAV)
 	{
+		ValidateBoundShader(Shader);
 		if (Bypass())
 		{
 			GetContext().RHISetUAVParameter(Shader, UAVIndex, UAV);
@@ -3029,9 +3079,10 @@ public:
 		ALLOC_COMMAND(FRHICommandBindClearMRTValues)(bClearColor, bClearDepth, bClearStencil);
 	}	
 
-	FORCEINLINE_DEBUGGABLE void SetGraphicsPipelineState(class FGraphicsPipelineState* GraphicsPipelineState)
+	FORCEINLINE_DEBUGGABLE void SetGraphicsPipelineState(class FGraphicsPipelineState* GraphicsPipelineState, const FBoundShaderStateInput& ShaderInput)
 	{
 		//check(IsOutsideRenderPass());
+		BoundShaderInput = ShaderInput;
 		if (Bypass())
 		{
 			extern RHI_API FRHIGraphicsPipelineState* ExecuteSetGraphicsPipelineState(class FGraphicsPipelineState* GraphicsPipelineState);
@@ -3735,76 +3786,41 @@ public:
 		return RHICreateBlendState(Initializer);
 	}
 	
-	FORCEINLINE FPixelShaderRHIRef CreatePixelShader(const TArray<uint8>& Code)
-	{
-		LLM_SCOPE(ELLMTag::Shaders);
-		return GDynamicRHI->CreatePixelShader_RenderThread(*this, Code);
-	}
-	
-	FORCEINLINE FPixelShaderRHIRef CreatePixelShader(FRHIShaderLibrary* Library, FSHAHash Hash)
-	{
-		LLM_SCOPE(ELLMTag::Shaders);
-		return GDynamicRHI->CreatePixelShader_RenderThread(*this, Library, Hash);
-	}
-	
-	FORCEINLINE FVertexShaderRHIRef CreateVertexShader(const TArray<uint8>& Code)
-	{
-		LLM_SCOPE(ELLMTag::Shaders);
-		return GDynamicRHI->CreateVertexShader_RenderThread(*this, Code);
-	}
-	
-	FORCEINLINE FVertexShaderRHIRef CreateVertexShader(FRHIShaderLibrary* Library, FSHAHash Hash)
-	{
-		LLM_SCOPE(ELLMTag::Shaders);
-		return GDynamicRHI->CreateVertexShader_RenderThread(*this, Library, Hash);
-	}
-	
-	FORCEINLINE FHullShaderRHIRef CreateHullShader(const TArray<uint8>& Code)
-	{
-		LLM_SCOPE(ELLMTag::Shaders);
-		return GDynamicRHI->CreateHullShader_RenderThread(*this, Code);
-	}
-	
-	FORCEINLINE FHullShaderRHIRef CreateHullShader(FRHIShaderLibrary* Library, FSHAHash Hash)
+	FORCEINLINE FPixelShaderRHIRef CreatePixelShader(TArrayView<const uint8> Code, const FSHAHash& Hash)
 	{
 		LLM_SCOPE(ELLMTag::Shaders);
-		return GDynamicRHI->CreateHullShader_RenderThread(*this, Library, Hash);
+		return GDynamicRHI->CreatePixelShader_RenderThread(*this, Code, Hash);
 	}
 	
-	FORCEINLINE FDomainShaderRHIRef CreateDomainShader(const TArray<uint8>& Code)
-	{
-		LLM_SCOPE(ELLMTag::Shaders);
-		return GDynamicRHI->CreateDomainShader_RenderThread(*this, Code);
-	}
 	
-	FORCEINLINE FDomainShaderRHIRef CreateDomainShader(FRHIShaderLibrary* Library, FSHAHash Hash)
+	FORCEINLINE FVertexShaderRHIRef CreateVertexShader(TArrayView<const uint8> Code, const FSHAHash& Hash)
 	{
 		LLM_SCOPE(ELLMTag::Shaders);
-		return GDynamicRHI->CreateDomainShader_RenderThread(*this, Library, Hash);
+		return GDynamicRHI->CreateVertexShader_RenderThread(*this, Code, Hash);
 	}
 	
-	FORCEINLINE FGeometryShaderRHIRef CreateGeometryShader(const TArray<uint8>& Code)
+	FORCEINLINE FHullShaderRHIRef CreateHullShader(TArrayView<const uint8> Code, const FSHAHash& Hash)
 	{
 		LLM_SCOPE(ELLMTag::Shaders);
-		return GDynamicRHI->CreateGeometryShader_RenderThread(*this, Code);
+		return GDynamicRHI->CreateHullShader_RenderThread(*this, Code, Hash);
 	}
 	
-	FORCEINLINE FGeometryShaderRHIRef CreateGeometryShader(FRHIShaderLibrary* Library, FSHAHash Hash)
+	FORCEINLINE FDomainShaderRHIRef CreateDomainShader(TArrayView<const uint8> Code, const FSHAHash& Hash)
 	{
 		LLM_SCOPE(ELLMTag::Shaders);
-		return GDynamicRHI->CreateGeometryShader_RenderThread(*this, Library, Hash);
+		return GDynamicRHI->CreateDomainShader_RenderThread(*this, Code, Hash);
 	}
 	
-	FORCEINLINE FComputeShaderRHIRef CreateComputeShader(const TArray<uint8>& Code)
+	FORCEINLINE FGeometryShaderRHIRef CreateGeometryShader(TArrayView<const uint8> Code, const FSHAHash& Hash)
 	{
 		LLM_SCOPE(ELLMTag::Shaders);
-		return GDynamicRHI->CreateComputeShader_RenderThread(*this, Code);
+		return GDynamicRHI->CreateGeometryShader_RenderThread(*this, Code, Hash);
 	}
 	
-	FORCEINLINE FComputeShaderRHIRef CreateComputeShader(FRHIShaderLibrary* Library, FSHAHash Hash)
+	FORCEINLINE FComputeShaderRHIRef CreateComputeShader(TArrayView<const uint8> Code, const FSHAHash& Hash)
 	{
 		LLM_SCOPE(ELLMTag::Shaders);
-		return GDynamicRHI->CreateComputeShader_RenderThread(*this, Library, Hash);
+		return GDynamicRHI->CreateComputeShader_RenderThread(*this, Code, Hash);
 	}
 	
 	FORCEINLINE FComputeFenceRHIRef CreateComputeFence(const FName& Name)
@@ -4255,7 +4271,7 @@ public:
 	{
 		GDynamicRHI->RHIMapStagingSurface_RenderThread(*this, Texture, Fence, OutData, OutWidth, OutHeight);
 	}
-
+	
 	FORCEINLINE void UnmapStagingSurface(FRHITexture* Texture)
 	{
 		GDynamicRHI->RHIUnmapStagingSurface_RenderThread(*this, Texture);
@@ -4514,7 +4530,7 @@ public:
 		 
 		return GDynamicRHI->RHIGetNativeDevice();
 	}
-
+	
 	FORCEINLINE void* GetNativeInstance()
 	{
 		QUICK_SCOPE_CYCLE_COUNTER(STAT_RHIMETHOD_GetNativeInstance_Flush);
@@ -4821,64 +4837,34 @@ struct FScopedCommandListWaitForTasks
 };
 
 
-FORCEINLINE FPixelShaderRHIRef RHICreatePixelShader(const TArray<uint8>& Code)
-{
-	return FRHICommandListExecutor::GetImmediateCommandList().CreatePixelShader(Code);
-}
-
-FORCEINLINE FPixelShaderRHIRef RHICreatePixelShader(FRHIShaderLibrary* Library, FSHAHash Hash)
-{
-	return FRHICommandListExecutor::GetImmediateCommandList().CreatePixelShader(Library, Hash);
-}
-
-FORCEINLINE FVertexShaderRHIRef RHICreateVertexShader(const TArray<uint8>& Code)
-{
-	return FRHICommandListExecutor::GetImmediateCommandList().CreateVertexShader(Code);
-}
-
-FORCEINLINE FVertexShaderRHIRef RHICreateVertexShader(FRHIShaderLibrary* Library, FSHAHash Hash)
-{
-	return FRHICommandListExecutor::GetImmediateCommandList().CreateVertexShader(Library, Hash);
-}
-
-FORCEINLINE FHullShaderRHIRef RHICreateHullShader(const TArray<uint8>& Code)
-{
-	return FRHICommandListExecutor::GetImmediateCommandList().CreateHullShader(Code);
-}
-
-FORCEINLINE FHullShaderRHIRef RHICreateHullShader(FRHIShaderLibrary* Library, FSHAHash Hash)
-{
-	return FRHICommandListExecutor::GetImmediateCommandList().CreateHullShader(Library, Hash);
-}
-
-FORCEINLINE FDomainShaderRHIRef RHICreateDomainShader(const TArray<uint8>& Code)
+FORCEINLINE FPixelShaderRHIRef RHICreatePixelShader(TArrayView<const uint8> Code, const FSHAHash& Hash)
 {
-	return FRHICommandListExecutor::GetImmediateCommandList().CreateDomainShader(Code);
+	return FRHICommandListExecutor::GetImmediateCommandList().CreatePixelShader(Code, Hash);
 }
 
-FORCEINLINE FDomainShaderRHIRef RHICreateDomainShader(FRHIShaderLibrary* Library, FSHAHash Hash)
+FORCEINLINE FVertexShaderRHIRef RHICreateVertexShader(TArrayView<const uint8> Code, const FSHAHash& Hash)
 {
-	return FRHICommandListExecutor::GetImmediateCommandList().CreateDomainShader(Library, Hash);
+	return FRHICommandListExecutor::GetImmediateCommandList().CreateVertexShader(Code, Hash);
 }
 
-FORCEINLINE FGeometryShaderRHIRef RHICreateGeometryShader(const TArray<uint8>& Code)
+FORCEINLINE FHullShaderRHIRef RHICreateHullShader(TArrayView<const uint8> Code, const FSHAHash& Hash)
 {
-	return FRHICommandListExecutor::GetImmediateCommandList().CreateGeometryShader(Code);
+	return FRHICommandListExecutor::GetImmediateCommandList().CreateHullShader(Code, Hash);
 }
 
-FORCEINLINE FGeometryShaderRHIRef RHICreateGeometryShader(FRHIShaderLibrary* Library, FSHAHash Hash)
+FORCEINLINE FDomainShaderRHIRef RHICreateDomainShader(TArrayView<const uint8> Code, const FSHAHash& Hash)
 {
-	return FRHICommandListExecutor::GetImmediateCommandList().CreateGeometryShader(Library, Hash);
+	return FRHICommandListExecutor::GetImmediateCommandList().CreateDomainShader(Code, Hash);
 }
 
-FORCEINLINE FComputeShaderRHIRef RHICreateComputeShader(const TArray<uint8>& Code)
+FORCEINLINE FGeometryShaderRHIRef RHICreateGeometryShader(TArrayView<const uint8> Code, const FSHAHash& Hash)
 {
-	return FRHICommandListExecutor::GetImmediateCommandList().CreateComputeShader(Code);
+	return FRHICommandListExecutor::GetImmediateCommandList().CreateGeometryShader(Code, Hash);
 }
 
-FORCEINLINE FComputeShaderRHIRef RHICreateComputeShader(FRHIShaderLibrary* Library, FSHAHash Hash)
+FORCEINLINE FComputeShaderRHIRef RHICreateComputeShader(TArrayView<const uint8> Code, const FSHAHash& Hash)
 {
-	return FRHICommandListExecutor::GetImmediateCommandList().CreateComputeShader(Library, Hash);
+	return FRHICommandListExecutor::GetImmediateCommandList().CreateComputeShader(Code, Hash);
 }
 
 FORCEINLINE FComputeFenceRHIRef RHICreateComputeFence(const FName& Name)
diff --git a/Engine/Source/Runtime/RHI/Public/RHIContext.h b/Engine/Source/Runtime/RHI/Public/RHIContext.h
--- a/Engine/Source/Runtime/RHI/Public/RHIContext.h
+++ b/Engine/Source/Runtime/RHI/Public/RHIContext.h
@@ -55,10 +55,10 @@ public:
 		checkf(UniformBuffer, TEXT("Attemped to assign a null uniform buffer to the global uniform buffer bindings."));
 		const FRHIUniformBufferLayout& Layout = UniformBuffer->GetLayout();
 		const FUniformBufferStaticSlot Slot = Layout.StaticSlot;
-		checkf(IsUniformBufferStaticSlotValid(Slot), TEXT("Attempted to set a global uniform buffer %s with an invalid slot."), *Layout.GetDebugName().ToString());
+		checkf(IsUniformBufferStaticSlotValid(Slot), TEXT("Attempted to set a global uniform buffer %s with an invalid slot."), *Layout.GetDebugName());
 
 #if VALIDATE_UNIFORM_BUFFER_GLOBAL_BINDINGS
-		ensureMsgf(INDEX_NONE == Slots.Find(Slot), TEXT("Uniform Buffer %s was added twice to the binding array."), *Layout.GetDebugName().ToString());
+		ensureMsgf(INDEX_NONE == Slots.Find(Slot), TEXT("Uniform Buffer %s was added twice to the binding array."), *Layout.GetDebugName());
 #endif
 
 		Slots.Add(Slot);
