Fix for signature checking related issues

diff --git a/Engine/Source/Runtime/PakFile/Private/IPlatformFilePak.cpp b/Engine/Source/Runtime/PakFile/Private/IPlatformFilePak.cpp
--- a/Engine/Source/Runtime/PakFile/Private/IPlatformFilePak.cpp
+++ b/Engine/Source/Runtime/PakFile/Private/IPlatformFilePak.cpp
@@ -187,17 +187,35 @@ FPakCustomEncryptionDelegate& FPakPlatformFile::GetPakCustomEncryptionDelegate()
 	return Delegate;
 }
 
+FPakPlatformFile::FPakSigningFailureHandlerData& FPakPlatformFile::GetPakSigningFailureHandlerData()
+{
+	static FPakSigningFailureHandlerData Instance;
+	return Instance;
+}
+
 // Needs to be a global as multiple threads may try to access the failed delegate
-static FPakChunkSignatureCheckFailedHandler SignatureFailedDelegate;
 FPakChunkSignatureCheckFailedHandler& FPakPlatformFile::GetPakChunkSignatureCheckFailedHandler()
 {
-	return SignatureFailedDelegate;
+	return GetPakSigningFailureHandlerData().ChunkSignatureCheckFailedDelegate;
 }
 
 FPakMasterSignatureTableCheckFailureHandler& FPakPlatformFile::GetPakMasterSignatureTableCheckFailureHandler()
 {
-	static FPakMasterSignatureTableCheckFailureHandler Delegate;
-	return Delegate;
+	return GetPakSigningFailureHandlerData().MasterSignatureTableCheckFailedDelegate;
+}
+
+void FPakPlatformFile::BroadcastPakChunkSignatureCheckFailure(const FPakChunkSignatureCheckFailedData& InData)
+{
+	FPakSigningFailureHandlerData& HandlerData = GetPakSigningFailureHandlerData();
+	FScopeLock Lock(&HandlerData.Lock);
+	HandlerData.ChunkSignatureCheckFailedDelegate.Broadcast(InData);
+}
+
+void FPakPlatformFile::BroadcastPakMasterSignatureTableCheckFailure(const FString& InFilename)
+{
+	FPakSigningFailureHandlerData& HandlerData = GetPakSigningFailureHandlerData();
+	FScopeLock Lock(&HandlerData.Lock);
+	HandlerData.MasterSignatureTableCheckFailedDelegate.Broadcast(InFilename);
 }
 
 void FPakPlatformFile::GetFilenamesInChunk(const FString& InPakFilename, const TArray<int32>& InChunkIDs, TArray<FString>& OutFileList) 
@@ -308,7 +326,7 @@ TSharedPtr<const struct FPakSignatureFile, ESPMode::ThreadSafe> FPakPlatformFile
 		else
 		{
 			UE_LOG(LogPakFile, Warning, TEXT("Couldn't find pak signature file '%s'"), InFilename);
-			FPakPlatformFile::GetPakMasterSignatureTableCheckFailureHandler().Broadcast(InFilename);
+			BroadcastPakMasterSignatureTableCheckFailure(InFilename);
 		}
 	}
 
@@ -3593,7 +3611,7 @@ void FPakPrecacher::DoSignatureCheck(bool bWasCanceled, IAsyncReadRequest* Reque
 				}
 
 				FPakChunkSignatureCheckFailedData FailedData(PakData->Name.ToString(), HashCache[SignedChunkIndex % MaxHashesToCache], ThisHash, SignatureIndex);
-				FPakPlatformFile::GetPakChunkSignatureCheckFailedHandler().Broadcast(FailedData);
+				FPakPlatformFile::BroadcastPakChunkSignatureCheckFailure(FailedData);
 			}
 		}
 
@@ -5155,10 +5173,16 @@ bool FPakFile::Check()
 	// against the signature table and fire the handler
 	if (bSigned)
 	{
-		FDelegateHandle DelegateHandle = FPakPlatformFile::GetPakChunkSignatureCheckFailedHandler().AddLambda([&bSuccess](const FPakChunkSignatureCheckFailedData&)
-		{ 
-			bSuccess = false; 
-		});
+		FDelegateHandle DelegateHandle;
+		FPakPlatformFile::FPakSigningFailureHandlerData& HandlerData = FPakPlatformFile::GetPakSigningFailureHandlerData();
+
+		{
+			FScopeLock Lock(&HandlerData.Lock);
+			DelegateHandle = HandlerData.ChunkSignatureCheckFailedDelegate.AddLambda([&bSuccess](const FPakChunkSignatureCheckFailedData&)
+			{
+				bSuccess = false;
+			});
+		}
 
 		int64 CurrentPos = 0;
 		const int64 Size = PakReader.TotalSize();
@@ -5170,7 +5194,11 @@ bool FPakFile::Check()
 			CurrentPos += FPakInfo::MaxChunkDataSize;
 		}
 
-		FPakPlatformFile::GetPakChunkSignatureCheckFailedHandler().Remove(DelegateHandle);
+		if (DelegateHandle.IsValid())
+		{
+			FScopeLock Lock(&HandlerData.Lock);
+			HandlerData.ChunkSignatureCheckFailedDelegate.Remove(DelegateHandle);
+		}
 	}
 	else
 	{
diff --git a/Engine/Source/Runtime/PakFile/Private/SignedArchiveReader.cpp b/Engine/Source/Runtime/PakFile/Private/SignedArchiveReader.cpp
--- a/Engine/Source/Runtime/PakFile/Private/SignedArchiveReader.cpp
+++ b/Engine/Source/Runtime/PakFile/Private/SignedArchiveReader.cpp
@@ -251,7 +251,7 @@ bool FChunkCacheWorker::CheckSignature(const FChunkRequest& ChunkInfo)
 			}
 
 			const FPakChunkSignatureCheckFailedData Data(Reader->GetArchiveName(), Signatures->ChunkHashes[ChunkInfo.Index], ChunkHash, ChunkInfo.Index);
-			FPakPlatformFile::GetPakChunkSignatureCheckFailedHandler().Broadcast(Data);
+			FPakPlatformFile::BroadcastPakChunkSignatureCheckFailure(Data);
 		}
 	}
 	
diff --git a/Engine/Source/Runtime/PakFile/Public/IPlatformFilePak.h b/Engine/Source/Runtime/PakFile/Public/IPlatformFilePak.h
--- a/Engine/Source/Runtime/PakFile/Public/IPlatformFilePak.h
+++ b/Engine/Source/Runtime/PakFile/Public/IPlatformFilePak.h
@@ -2434,12 +2434,30 @@ public:
 	// Access static delegate for custom encryption
 	static FPakCustomEncryptionDelegate& GetPakCustomEncryptionDelegate();
 
+	struct FPakSigningFailureHandlerData
+	{
+		FCriticalSection Lock;
+		FPakChunkSignatureCheckFailedHandler ChunkSignatureCheckFailedDelegate;
+		FPakMasterSignatureTableCheckFailureHandler MasterSignatureTableCheckFailedDelegate;
+	};
+
 	// Access static delegate for handling a pak signature check failure
+	static FPakSigningFailureHandlerData& GetPakSigningFailureHandlerData();
+	
+	// Access static delegate for handling a pak signature check failure
+	UE_DEPRECATED(4.25, "GetPakChunkSignatureCheckFailedHandler is not thread safe, so please migrate to using GetPakSigningFailureHandlerData and locking the critical section around any use of the delegates")
 	static FPakChunkSignatureCheckFailedHandler& GetPakChunkSignatureCheckFailedHandler();
 
 	// Access static delegate for handling a pak signature check failure
+	UE_DEPRECATED(4.25, "GetPakMasterSignatureTableCheckFailureHandler is not thread safe, so please migrate to using GetPakSigningFailureHandlerData and locking the critical section around any use of the delegates")
 	static FPakMasterSignatureTableCheckFailureHandler& GetPakMasterSignatureTableCheckFailureHandler();
 
+	// Broadacast a signature check failure through any registered delegates in a thread safe way
+	static void BroadcastPakChunkSignatureCheckFailure(const FPakChunkSignatureCheckFailedData& InData);
+
+	// Broadacast a master signature table failure through any registered delegates in a thread safe way
+	static void BroadcastPakMasterSignatureTableCheckFailure(const FString& InFilename);
+
 	// Get a list of which files live in a given chunk
 	void GetFilenamesInChunk(const FString& InPakFilename, const TArray<int32>& InChunkIDs, TArray<FString>& OutFileList);
 	void GetFilenamesInPakFile(const FString& InPakFilename, TArray<FString>& OutFileList) ;
@@ -2558,7 +2576,7 @@ struct FPakSignatureFile
 			}
 		}
 
-		FPakPlatformFile::GetPakMasterSignatureTableCheckFailureHandler().Broadcast(InFilename);
+		FPakPlatformFile::BroadcastPakMasterSignatureTableCheckFailure(InFilename);
 		return false;
 	}
 
