Optimize Pak file indices by computing the PathHashIndex in unrealpak rather than using strings at runtime.

diff --git a/Engine/Source/Runtime/PakFile/Private/IPlatformFilePak.cpp b/Engine/Source/Runtime/PakFile/Private/IPlatformFilePak.cpp
--- a/Engine/Source/Runtime/PakFile/Private/IPlatformFilePak.cpp
+++ b/Engine/Source/Runtime/PakFile/Private/IPlatformFilePak.cpp
@@ -11,6 +11,7 @@
 #include "Misc/CoreDelegates.h"
 #include "Misc/App.h"
 #include "Modules/ModuleManager.h"
+#include "Misc/ConfigCacheIni.h"
 #include "Misc/SecureHash.h"
 #include "HAL/FileManagerGeneric.h"
 #include "HAL/IPlatformFileModule.h"
@@ -20,6 +21,7 @@
 #include "Async/AsyncFileHandle.h"
 #include "Templates/Greater.h"
 #include "Serialization/ArchiveProxy.h"
+#include "Serialization/MemoryWriter.h"
 #include "Misc/Base64.h"
 #include "HAL/DiskUtilizationTracker.h"
 #include "Stats/StatsMisc.h"
@@ -30,7 +32,7 @@
 #endif
 #include "ProfilingDebugging/CsvProfiler.h"
 #include "Misc/Fnv.h"
-
+#include "ProfilingDebugging/ScopedTimers.h"
 #include "Async/MappedFileHandle.h"
 #include "IO/IoDispatcher.h"
 
@@ -41,10 +43,6 @@ DEFINE_LOG_CATEGORY(LogPakFile);
 DEFINE_STAT(STAT_PakFile_Read);
 DEFINE_STAT(STAT_PakFile_NumOpenHandles);
 
-IMPLEMENT_TYPE_LAYOUT(FPakCompressedBlock);
-IMPLEMENT_TYPE_LAYOUT(FPakEntry);
-IMPLEMENT_TYPE_LAYOUT(FPakFileData);
-
 CSV_DECLARE_CATEGORY_MODULE_EXTERN(CORE_API, FileIO);
 
 #ifndef DISABLE_NONUFS_INI_WHEN_COOKED
@@ -218,7 +216,13 @@ void FPakPlatformFile::BroadcastPakMasterSignatureTableCheckFailure(const FStrin
 	HandlerData.MasterSignatureTableCheckFailedDelegate.Broadcast(InFilename);
 }
 
-void FPakPlatformFile::GetFilenamesInChunk(const FString& InPakFilename, const TArray<int32>& InChunkIDs, TArray<FString>& OutFileList) 
+FPakSetIndexSettings& FPakPlatformFile::GetPakSetIndexSettingsDelegate()
+{
+	static FPakSetIndexSettings Delegate;
+	return Delegate;
+}
+
+void FPakPlatformFile::GetPrunedFilenamesInChunk(const FString& InPakFilename, const TArray<int32>& InChunkIDs, TArray<FString>& OutFileList)
 {
 	TArray<FPakListEntry> Paks;
 	GetMountedPaks(Paks);
@@ -227,13 +231,13 @@ void FPakPlatformFile::GetFilenamesInChunk(const FString& InPakFilename, const T
 	{
 		if (Pak.PakFile && Pak.PakFile->GetFilename() == InPakFilename)
 		{
-			Pak.PakFile->GetFilenamesInChunk(InChunkIDs, OutFileList);
+			Pak.PakFile->GetPrunedFilenamesInChunk(InChunkIDs, OutFileList);
 			break;
 		}
 	}
 }
 
-void FPakPlatformFile::GetFilenamesInPakFile(const FString& InPakFilename, TArray<FString>& OutFileList)
+void FPakPlatformFile::GetPrunedFilenamesInPakFile(const FString& InPakFilename, TArray<FString>& OutFileList)
 {
 	TArray<FPakListEntry> Paks;
 	GetMountedPaks(Paks);
@@ -242,7 +246,7 @@ void FPakPlatformFile::GetFilenamesInPakFile(const FString& InPakFilename, TArra
 	{
 		if (Pak.PakFile && Pak.PakFile->GetFilename() == InPakFilename)
 		{
-			Pak.PakFile->GetFilenames(OutFileList);
+			Pak.PakFile->GetPrunedFilenames(OutFileList);
 			break;
 		}
 	}
@@ -4268,7 +4272,6 @@ public:
 IAsyncReadFileHandle* FPakPlatformFile::OpenAsyncRead(const TCHAR* Filename)
 {
 	CSV_SCOPED_TIMING_STAT(FileIO, PakOpenAsyncRead);
-//	check(GConfig);
 #if USE_PAK_PRECACHE
 	if (FPlatformProcess::SupportsMultithreading() && GPakCache_Enable > 0)
 	{
@@ -4788,19 +4791,18 @@ bool FPakPlatformFile::IsNonPakFilenameAllowed(const FString& InFilename)
 FPakFile::FPakFile(const TCHAR* Filename, bool bIsSigned)
 	: PakFilename(Filename)
 	, PakFilenameName(Filename)
-	, FilenameHashesIndex(nullptr)
-	, FilenameHashesIndices(nullptr)
-	, FilenameHashes(nullptr)
-	, MiniPakEntriesOffsets(nullptr)
-	, MiniPakEntries(nullptr)
+	, PathHashSeed(0)
 	, NumEntries(0)
 	, CachedTotalSize(0)
 	, bSigned(bIsSigned)
 	, bIsValid(false)
-	, bFilenamesRemoved(false)
+	, bHasPathHashIndex(false)
+	, bHasFullDirectoryIndex(false)
+#if ENABLE_PAKFILE_RUNTIME_PRUNING
+	, bWillPruneDirectoryIndex(false)
+	, bNeedsLegacyPruning(false)
+#endif
 	, PakchunkIndex(GetPakchunkIndexFromPakFile(Filename))
-	, bAttemptedPakEntryShrink(false)
-	, bAttemptedPakFilenameUnload(false)
  	, MappedFileHandle(nullptr)
 	, CacheType(FPakFile::ECacheType::Shared)
 	, CacheIndex(-1)
@@ -4818,19 +4820,18 @@ FPakFile::FPakFile(const TCHAR* Filename, bool bIsSigned)
 FPakFile::FPakFile(IPlatformFile* LowerLevel, const TCHAR* Filename, bool bIsSigned, bool bLoadIndex)
 	: PakFilename(Filename)
 	, PakFilenameName(Filename)
-	, FilenameHashesIndex(nullptr)
-	, FilenameHashesIndices(nullptr)
-	, FilenameHashes(nullptr)
-	, MiniPakEntriesOffsets(nullptr)
-	, MiniPakEntries(nullptr)
+	, PathHashSeed(0)
 	, NumEntries(0)
 	, CachedTotalSize(0)
 	, bSigned(bIsSigned)
 	, bIsValid(false)
-	, bFilenamesRemoved(false)
+	, bHasPathHashIndex(false)
+	, bHasFullDirectoryIndex(false)
+#if ENABLE_PAKFILE_RUNTIME_PRUNING
+	, bWillPruneDirectoryIndex(false)
+	, bNeedsLegacyPruning(false)
+#endif
 	, PakchunkIndex(GetPakchunkIndexFromPakFile(Filename))
-	, bAttemptedPakEntryShrink(false)
-	, bAttemptedPakFilenameUnload(false)
 	, MappedFileHandle(nullptr)
 	, CacheType(FPakFile::ECacheType::Shared)
 	, CacheIndex(-1)
@@ -4846,15 +4847,16 @@ FPakFile::FPakFile(IPlatformFile* LowerLevel, const TCHAR* Filename, bool bIsSig
 
 #if WITH_EDITOR
 FPakFile::FPakFile(FArchive* Archive)
-	: FilenameHashesIndex(nullptr)
-	, FilenameHashesIndices(nullptr)
-	, FilenameHashes(nullptr)
-	, MiniPakEntriesOffsets(nullptr)
-	, MiniPakEntries(nullptr)
+	: PathHashSeed(0)
 	, NumEntries(0)
 	, bSigned(false)
 	, bIsValid(false)
-	, bFilenamesRemoved(false)
+	, bHasPathHashIndex(false)
+	, bHasFullDirectoryIndex(false)
+#if ENABLE_PAKFILE_RUNTIME_PRUNING
+	, bWillPruneDirectoryIndex(false)
+	, bNeedsLegacyPruning(false)
+#endif
 	, PakchunkIndex(INDEX_NONE)
 	, MappedFileHandle(nullptr)
 	, CacheType(FPakFile::ECacheType::Shared)
@@ -4868,11 +4870,6 @@ FPakFile::FPakFile(FArchive* Archive)
 FPakFile::~FPakFile()
 {
 	delete MappedFileHandle;
-	delete[] MiniPakEntries;
-	delete[] MiniPakEntriesOffsets;
-	delete[] FilenameHashes;
-	delete[] FilenameHashesIndices;
-	delete[] FilenameHashesIndex;
 }
 
 bool FPakFile::PassedSignatureChecks() const
@@ -4982,762 +4979,1208 @@ void FPakFile::Initialize(FArchive* Reader, bool bLoadIndex)
 
 void FPakFile::LoadIndex(FArchive* Reader)
 {
-	if (CachedTotalSize < (Info.IndexOffset + Info.IndexSize))
+	if (Info.Version >= FPakInfo::PakFile_Version_PathHashIndex)
 	{
-		UE_LOG(LogPakFile, Fatal, TEXT("Corrupted index offset in pak file."));
+		if (!LoadIndexInternal(Reader))
+		{
+			// Index loading failed. Try again
+			if (!LoadIndexInternal(Reader))
+			{
+				UE_LOG(LogPakFile, Fatal, TEXT("Corrupt pak index detected on pak file: %s"), *PakFilename);
+			}
+		}
 	}
 	else
 	{
-		if (Info.Version >= FPakInfo::PakFile_Version_FrozenIndex && Info.bIndexIsFrozen)
+		SCOPED_BOOT_TIMING("PakFile_LoadLegacy");
+		if (!LoadLegacyIndex(Reader))
 		{
-			SCOPED_BOOT_TIMING("PakFile_LoadFrozen");
-
-			// read frozen data
-			Reader->Seek(Info.IndexOffset);
-			int32 FrozenSize = Info.IndexSize;
-
-			// read in the index, etc data in one lump
-			void* DataMemory = FMemory::Malloc(FrozenSize);
-			Reader->Serialize(DataMemory, FrozenSize);
-			Data = TUniquePtr<FPakFileData>((FPakFileData*)DataMemory);
-
-			// cache the number of entries
-			NumEntries = Data->Files.Num();
-			// @todo loadtime: it is nice to serialize the mountpoint right into the Data so that IndexSize is right here
-			// but it takes this to copy it out, because it's too painful for the string manipulation when dealing with
-			// MemoryImageString everywhere MountPoint is used
-			MountPoint = Data->MountPoint;
-		}
-		else
-		{
-			// Load index into memory first.
-
-			// If we encounter a corrupt index, try again but gather some extra debugging information so we can try and understand where it failed.
-			bool bFirstPass = true;
-			TArray<uint8> IndexData;
-
-			while (true)
+			// Index loading failed. Try again
+			if (!LoadLegacyIndex(Reader))
 			{
-				Reader->Seek(Info.IndexOffset);
-				IndexData.Empty(); // The next SetNum makes this Empty() logically redundant, but we want to try and force a memory reallocation for the re-attempt
-				IndexData.SetNum(Info.IndexSize);
-				{
-					SCOPED_BOOT_TIMING("PakFile_LoadIndex");
-					Reader->Serialize(IndexData.GetData(), Info.IndexSize);
-				}
-
-				FSHAHash EncryptedDataHash;
-				if (!bFirstPass)
-				{
-					FSHA1::HashBuffer(IndexData.GetData(), IndexData.Num(), EncryptedDataHash.Hash);
-				}
-
-				// Decrypt if necessary
-				if (Info.bEncryptedIndex)
-				{
-					DecryptData(IndexData.GetData(), Info.IndexSize, Info.EncryptionKeyGuid);
-				}
-				// Check SHA1 value.
-				FSHAHash ComputedHash;
-				{
-					SCOPED_BOOT_TIMING("PakFile_Hash");
-					FSHA1::HashBuffer(IndexData.GetData(), IndexData.Num(), ComputedHash.Hash);
-				}
-				if (Info.IndexHash != ComputedHash)
-				{
-					if (bFirstPass)
-					{
-						UE_LOG(LogPakFile, Log, TEXT("Corrupt pak index detected!"));
-						UE_LOG(LogPakFile, Log, TEXT(" Filename: %s"), *PakFilename);
-						UE_LOG(LogPakFile, Log, TEXT(" Encrypted: %d"), Info.bEncryptedIndex);
-						UE_LOG(LogPakFile, Log, TEXT(" Total Size: %d"), Reader->TotalSize());
-						UE_LOG(LogPakFile, Log, TEXT(" Index Offset: %d"), Info.IndexOffset);
-						UE_LOG(LogPakFile, Log, TEXT(" Index Size: %d"), Info.IndexSize);
-						UE_LOG(LogPakFile, Log, TEXT(" Stored Index Hash: %s"), *Info.IndexHash.ToString());
-						UE_LOG(LogPakFile, Log, TEXT(" Computed Index Hash [Pass 0]: %s"), *ComputedHash.ToString());
-						bFirstPass = false;
-					}
-					else
-					{
-						UE_LOG(LogPakFile, Log, TEXT(" Computed Index Hash [Pass 1]: %s"), *ComputedHash.ToString());
-						UE_LOG(LogPakFile, Log, TEXT(" Encrypted Index Hash: %s"), *EncryptedDataHash.ToString());
-
-						// Compute an SHA1 hash of the whole file so we can tell if the file was modified on disc (obviously as long as this isn't an  IO bug that gives us the same bogus data again)
-						FSHA1 FileHash;
-						Reader->Seek(0);
-						int64 Remaining = Reader->TotalSize();
-						TArray<uint8> WorkingBuffer;
-						WorkingBuffer.SetNum(64 * 1024);
-						while (Remaining > 0)
-						{
-							int64 ToProcess = FMath::Min((int64)WorkingBuffer.Num(), Remaining);
-							Reader->Serialize(WorkingBuffer.GetData(), ToProcess);
-							FileHash.Update(WorkingBuffer.GetData(), ToProcess);
-							Remaining -= ToProcess;
-						}
-
-						FileHash.Final();
-						FSHAHash FinalFileHash;
-						FileHash.GetHash(FinalFileHash.Hash);
-						UE_LOG(LogPakFile, Log, TEXT(" File Hash: %s"), *FinalFileHash.ToString());
-
-						UE_LOG(LogPakFile, Fatal, TEXT("Corrupted index in pak file (SHA hash mismatch)."));
-					}
-				}
-				else
-				{
-					if (!bFirstPass)
-					{
-						UE_LOG(LogPakFile, Log, TEXT("Pak index corruption appears to have recovered on the second attempt!"));
-					}
-					break;
-				}
+				UE_LOG(LogPakFile, Fatal, TEXT("Corrupt pak index detected on pak file: %s"), *PakFilename);
 			}
+		}
+	}
+}
 
-			FMemoryReader IndexReader(IndexData);
-
-			// Read the default mount point and all entries.
-			NumEntries = 0;
-			IndexReader << MountPoint;
-			IndexReader << NumEntries;
+bool FPakFile::LoadIndexInternal(FArchive* Reader)
+{
+	bHasPathHashIndex = false;
+	bHasFullDirectoryIndex = false;
+#if ENABLE_PAKFILE_RUNTIME_PRUNING
+	bNeedsLegacyPruning = false;
+	bWillPruneDirectoryIndex = false;
+#endif
 
-			MakeDirectoryFromPath(MountPoint);
-			// Allocate enough memory to hold all entries (and not reallocate while they're being added to it).
+	if (CachedTotalSize < (Info.IndexOffset + Info.IndexSize))
+	{
+		UE_LOG(LogPakFile, Fatal, TEXT("Corrupted index offset in pak file."));
+		return false;
+	}
+	TArray<uint8> PrimaryIndexData;
+	Reader->Seek(Info.IndexOffset);
+	PrimaryIndexData.SetNum(Info.IndexSize);
+	{
+		SCOPED_BOOT_TIMING("PakFile_LoadPrimaryIndex");
+		Reader->Serialize(PrimaryIndexData.GetData(), Info.IndexSize);
+	}
 
-			Data = MakeUnique<FPakFileData>();
-			Data->Files.Empty(NumEntries);
-			Data->Files.AddZeroed(NumEntries);
+	FSHAHash ComputedHash;
+	{
+		SCOPED_BOOT_TIMING("PakFile_HashPrimaryIndex");
+		if (!DecryptAndValidateIndex(Reader, PrimaryIndexData, Info.IndexHash, ComputedHash))
+		{
+			UE_LOG(LogPakFile, Log, TEXT("Corrupt pak PrimaryIndex detected!"));
+			UE_LOG(LogPakFile, Log, TEXT(" Filename: %s"), *PakFilename);
+			UE_LOG(LogPakFile, Log, TEXT(" Encrypted: %d"), Info.bEncryptedIndex);
+			UE_LOG(LogPakFile, Log, TEXT(" Total Size: %d"), Reader->TotalSize());
+			UE_LOG(LogPakFile, Log, TEXT(" Index Offset: %d"), Info.IndexOffset);
+			UE_LOG(LogPakFile, Log, TEXT(" Index Size: %d"), Info.IndexSize);
+			UE_LOG(LogPakFile, Log, TEXT(" Stored Index Hash: %s"), *Info.IndexHash.ToString());
+			UE_LOG(LogPakFile, Log, TEXT(" Computed Index Hash: %s"), *ComputedHash.ToString());
+			return false;
+		}
+	}
 
-			for (int32 EntryIndex = 0; EntryIndex < NumEntries; EntryIndex++)
-			{
-				// Serialize from memory.
-				FPakEntry& Entry = Data->Files[EntryIndex];
-				FString Filename;
-				IndexReader << Filename;
-				{
-					SCOPED_BOOT_TIMING("PakFile_SerializeEntry");
-					Entry.Serialize(IndexReader, Info.Version);
-				}
+	FMemoryReader PrimaryIndexReader(PrimaryIndexData);
 
-				// Construct Index of all directories in pak file.
-				FString Path = FPaths::GetPath(Filename);
-				MakeDirectoryFromPath(Path);
-				FPakDirectory* Directory = Data->Index.Find(Path);
-				if (Directory != NULL)
-				{
-					Directory->Add(FPaths::GetCleanFilename(Filename), EntryIndex);
-				}
-				else
-				{
-					FPakDirectory& NewDirectory = Data->Index.Add(Path);
-					NewDirectory.Add(FPaths::GetCleanFilename(Filename), EntryIndex);
+	// Read the scalar data (mount point, numentries, etc) and all entries.
+	NumEntries = 0;
+	PrimaryIndexReader << MountPoint;
+	MakeDirectoryFromPath(MountPoint);
+	PrimaryIndexReader << NumEntries;
+	PrimaryIndexReader << PathHashSeed;
 
-					// add the parent directories up to the mount point
-					while (MountPoint != Path)
-					{
-						Path.LeftInline(Path.Len() - 1);
-						int32 Offset = 0;
-						if (Path.FindLastChar('/', Offset))
-						{
-							Path.LeftInline(Offset);
-							MakeDirectoryFromPath(Path);
-							if (Data->Index.Find(Path) == NULL)
-							{
-								Data->Index.Add(Path);
-							}
-						}
-						else
-						{
-							Path = MountPoint;
-						}
-					}
-				}
-			}
-		}
+	bool bReaderHasPathHashIndex = false;
+	int64 PathHashIndexOffset = INDEX_NONE;
+	int64 PathHashIndexSize = 0;
+	FSHAHash PathHashIndexHash;
+	PrimaryIndexReader << bReaderHasPathHashIndex;
+	if (bReaderHasPathHashIndex)
+	{
+		PrimaryIndexReader << PathHashIndexOffset;
+		PrimaryIndexReader << PathHashIndexSize;
+		PrimaryIndexReader << PathHashIndexHash;
+		bReaderHasPathHashIndex = bReaderHasPathHashIndex && PathHashIndexOffset != INDEX_NONE;
 	}
-}
 
-bool FPakFile::Check()
-{
-	UE_LOG(LogPakFile, Display, TEXT("Checking pak file \"%s\". This may take a while..."), *PakFilename);
-	double StartTime = FPlatformTime::Seconds();
+	bool bReaderHasFullDirectoryIndex = false;
+	int64 FullDirectoryIndexOffset = INDEX_NONE;
+	int64 FullDirectoryIndexSize = 0;
+	FSHAHash FullDirectoryIndexHash;
+	PrimaryIndexReader << bReaderHasFullDirectoryIndex;
+	if (bReaderHasFullDirectoryIndex)
+	{
+		PrimaryIndexReader << FullDirectoryIndexOffset;
+		PrimaryIndexReader << FullDirectoryIndexSize;
+		PrimaryIndexReader << FullDirectoryIndexHash;
+		bReaderHasFullDirectoryIndex = bReaderHasFullDirectoryIndex && FullDirectoryIndexOffset  != INDEX_NONE;
+	}
+	{
+		SCOPED_BOOT_TIMING("PakFile_SerializeEncodedEntries");
+		PrimaryIndexReader << EncodedPakEntries;
+	}
 
-	FArchive& PakReader = *GetSharedReader(NULL);
-	int32 ErrorCount = 0;
-	int32 FileCount = 0;
+	int32 FilesNum = 0;
+	PrimaryIndexReader << FilesNum;
+	if (FilesNum < 0)
+	{
+		// Should not be possible for any values in the PrimaryIndex to be invalid, since we verified the index hash
+		UE_LOG(LogPakFile, Log, TEXT("Corrupt pak PrimaryIndex detected!"));
+		UE_LOG(LogPakFile, Log, TEXT(" FilesNum: %d"), FilesNum);
+		return false;
+	}
+	Files.SetNum(FilesNum);
+	if (FilesNum > 0)
+	{
+		SCOPED_BOOT_TIMING("PakFile_SerializeUnencodedEntries");
+		FPakEntry* FileEntries = Files.GetData();
+		for (int32 FileIndex = 0; FileIndex < FilesNum; ++FileIndex)
+		{
+			FileEntries[FileIndex].Serialize(PrimaryIndexReader, Info.Version);
+		}
+	}
 
-	bool bSuccess = true;
+	// Decide which SecondaryIndex(es) to load
+	bool bWillUseFullDirectoryIndex;
+	bool bWillUsePathHashIndex;
+	bool bReadFullDirectoryIndex;
+	if (bReaderHasPathHashIndex && bReaderHasFullDirectoryIndex)
+	{
+		bWillUseFullDirectoryIndex = IsPakKeepFullDirectory();
+		bWillUsePathHashIndex = !bWillUseFullDirectoryIndex;
+#if ENABLE_PAKFILE_RUNTIME_PRUNING
+		bool bWantToReadFullDirectoryIndex = IsPakKeepFullDirectory() || IsPakValidatePruning() || IsPakDelayPruning();
+#else
+		bool bWantToReadFullDirectoryIndex = IsPakKeepFullDirectory();
+#endif
+		bReadFullDirectoryIndex = bReaderHasFullDirectoryIndex && bWantToReadFullDirectoryIndex;
+	}
+	else if (bReaderHasPathHashIndex)
+	{
+		bWillUsePathHashIndex = true;
+		bWillUseFullDirectoryIndex = false;
+		bReadFullDirectoryIndex = false;
+	}
+	else if (bReaderHasFullDirectoryIndex)
+	{
+		// We don't support creating the PathHash Index at runtime; we want to move to having only the PathHashIndex, so supporting not having it at all is not useful enough to write
+		bWillUsePathHashIndex = false;
+		bWillUseFullDirectoryIndex = true;
+		bReadFullDirectoryIndex = true;
+	}
+	else
+	{
+		// It should not be possible for PrimaryIndexes to be built without a PathHashIndex AND without a FullDirectoryIndex; CreatePakFile in UnrealPak.exe has a check statement for it.
+		UE_LOG(LogPakFile, Log, TEXT("Corrupt pak PrimaryIndex detected!"));
+		UE_LOG(LogPakFile, Log, TEXT(" bReaderHasPathHashIndex: false"));
+		UE_LOG(LogPakFile, Log, TEXT(" bReaderHasFullDirectoryIndex: false"));
+		return false;
+	}
 
-	// If the pak file is signed, we can do a fast check by just reading a single byte from the start of
-	// each signing block. The signed archive reader will bring in that whole signing block and compare
-	// against the signature table and fire the handler
-	if (bSigned)
+	// Load the Secondary Index(es)
+	TArray<uint8> PathHashIndexData;
+	FMemoryReader PathHashIndexReader(PathHashIndexData);
+	if (bWillUsePathHashIndex)
 	{
-		FDelegateHandle DelegateHandle;
-		FPakPlatformFile::FPakSigningFailureHandlerData& HandlerData = FPakPlatformFile::GetPakSigningFailureHandlerData();
+		if (PathHashIndexOffset < 0 || CachedTotalSize < (PathHashIndexOffset + PathHashIndexSize))
+		{
+			// Should not be possible for these values (which came from the PrimaryIndex) to be invalid, since we verified the index hash of the PrimaryIndex
+			UE_LOG(LogPakFile, Log, TEXT("Corrupt pak PrimaryIndex detected!"));
+			UE_LOG(LogPakFile, Log, TEXT(" Filename: %s"), *PakFilename);
+			UE_LOG(LogPakFile, Log, TEXT(" Total Size: %d"), Reader->TotalSize());
+			UE_LOG(LogPakFile, Log, TEXT(" PathHashIndexOffset : %d"), PathHashIndexOffset);
+			UE_LOG(LogPakFile, Log, TEXT(" PathHashIndexSize: %d"), PathHashIndexSize);
+			return false;
+		}
+		Reader->Seek(PathHashIndexOffset);
+		PathHashIndexData.SetNum(PathHashIndexSize);
+		{
+			SCOPED_BOOT_TIMING("PakFile_LoadPathHashIndex");
+			Reader->Serialize(PathHashIndexData.GetData(), PathHashIndexSize);
+		}
 
 		{
-			FScopeLock Lock(&HandlerData.Lock);
-			DelegateHandle = HandlerData.ChunkSignatureCheckFailedDelegate.AddLambda([&bSuccess](const FPakChunkSignatureCheckFailedData&)
+			SCOPED_BOOT_TIMING("PakFile_HashPathHashIndex");
+			if (!DecryptAndValidateIndex(Reader, PathHashIndexData, PathHashIndexHash, ComputedHash))
 			{
-				bSuccess = false;
-			});
+				UE_LOG(LogPakFile, Log, TEXT("Corrupt pak PathHashIndex detected!"));
+				UE_LOG(LogPakFile, Log, TEXT(" Filename: %s"), *PakFilename);
+				UE_LOG(LogPakFile, Log, TEXT(" Encrypted: %d"), Info.bEncryptedIndex);
+				UE_LOG(LogPakFile, Log, TEXT(" Total Size: %d"), Reader->TotalSize());
+				UE_LOG(LogPakFile, Log, TEXT(" Index Offset: %d"), FullDirectoryIndexOffset);
+				UE_LOG(LogPakFile, Log, TEXT(" Index Size: %d"), FullDirectoryIndexSize);
+				UE_LOG(LogPakFile, Log, TEXT(" Stored Index Hash: %s"), *PathHashIndexHash.ToString());
+				UE_LOG(LogPakFile, Log, TEXT(" Computed Index Hash: %s"), *ComputedHash.ToString());
+				return false;
+			}
 		}
 
-		int64 CurrentPos = 0;
-		const int64 Size = PakReader.TotalSize();
-		while (CurrentPos < Size)
 		{
-			PakReader.Seek(CurrentPos);
-			uint8 Byte = 0;
-			PakReader << Byte;
-			CurrentPos += FPakInfo::MaxChunkDataSize;
+			SCOPED_BOOT_TIMING("PakFile_SerializePathHashIndex");
+			PathHashIndexReader << PathHashIndex;
 		}
+		bHasPathHashIndex = true;
+	}
 
-		if (DelegateHandle.IsValid())
+	if (!bReadFullDirectoryIndex)
+	{
+		check(bWillUsePathHashIndex); // Need to confirm that we have read the PathHashIndex bytes
+		// Store the PrunedDirectoryIndex in our DirectoryIndex
 		{
-			FScopeLock Lock(&HandlerData.Lock);
-			HandlerData.ChunkSignatureCheckFailedDelegate.Remove(DelegateHandle);
+			SCOPED_BOOT_TIMING("PakFile_SerializePrunedDirectoryIndex");
+			PathHashIndexReader << DirectoryIndex;
 		}
+		bHasFullDirectoryIndex = false;
+#if ENABLE_PAKFILE_RUNTIME_PRUNING
+		bWillPruneDirectoryIndex = false;
+#endif
 	}
 	else
 	{
-		const bool bIncludeDeleted = true;
-		for (FPakFile::FFileIterator It(*this, bIncludeDeleted); It; ++It, ++FileCount)
+		if (CachedTotalSize < (FullDirectoryIndexOffset  + FullDirectoryIndexSize) ||
+			FullDirectoryIndexOffset  < 0)
 		{
-			const FPakEntry& Entry = It.Info();
-			if (Entry.IsDeleteRecord())
-			{
-				UE_LOG(LogPakFile, Verbose, TEXT("\"%s\" Deleted."), *It.Filename());
-				continue;
-			}
+			// Should not be possible for these values (which came from the PrimaryIndex) to be invalid, since we verified the index hash of the PrimaryIndex
+			UE_LOG(LogPakFile, Log, TEXT("Corrupt pak PrimaryIndex detected!"));
+			UE_LOG(LogPakFile, Log, TEXT(" Filename: %s"), *PakFilename);
+			UE_LOG(LogPakFile, Log, TEXT(" Total Size: %d"), Reader->TotalSize());
+			UE_LOG(LogPakFile, Log, TEXT(" FullDirectoryIndexOffset : %d"), FullDirectoryIndexOffset );
+			UE_LOG(LogPakFile, Log, TEXT(" FullDirectoryIndexSize: %d"), FullDirectoryIndexSize);
+			return false;
+		}
+		TArray<uint8> FullDirectoryIndexData;
+		Reader->Seek(FullDirectoryIndexOffset );
+		FullDirectoryIndexData.SetNum(FullDirectoryIndexSize);
+		{
+			SCOPED_BOOT_TIMING("PakFile_LoadDirectoryIndex");
+			Reader->Serialize(FullDirectoryIndexData.GetData(), FullDirectoryIndexSize);
+		}
 
-			void* FileContents = FMemory::Malloc(Entry.Size);
-			PakReader.Seek(Entry.Offset);
-			uint32 SerializedCrcTest = 0;
-			FPakEntry EntryInfo;
-			EntryInfo.Serialize(PakReader, GetInfo().Version);
-			if (EntryInfo != Entry)
+		{
+			SCOPED_BOOT_TIMING("PakFile_HashDirectoryIndex");
+			if (!DecryptAndValidateIndex(Reader, FullDirectoryIndexData, FullDirectoryIndexHash, ComputedHash))
 			{
-				UE_LOG(LogPakFile, Error, TEXT("Serialized hash mismatch for \"%s\"."), *It.Filename());
-				ErrorCount++;
+				UE_LOG(LogPakFile, Log, TEXT("Corrupt pak FullDirectoryIndex detected!"));
+				UE_LOG(LogPakFile, Log, TEXT(" Filename: %s"), *PakFilename);
+				UE_LOG(LogPakFile, Log, TEXT(" Encrypted: %d"), Info.bEncryptedIndex);
+				UE_LOG(LogPakFile, Log, TEXT(" Total Size: %d"), Reader->TotalSize());
+				UE_LOG(LogPakFile, Log, TEXT(" Index Offset: %d"), FullDirectoryIndexOffset);
+				UE_LOG(LogPakFile, Log, TEXT(" Index Size: %d"), FullDirectoryIndexSize);
+				UE_LOG(LogPakFile, Log, TEXT(" Stored Index Hash: %s"), *FullDirectoryIndexHash.ToString());
+				UE_LOG(LogPakFile, Log, TEXT(" Computed Index Hash: %s"), *ComputedHash.ToString());
+				return false;
 			}
-			PakReader.Serialize(FileContents, Entry.Size);
+		}
 
-			uint8 TestHash[20];
-			FSHA1::HashBuffer(FileContents, Entry.Size, TestHash);
-			if (FMemory::Memcmp(TestHash, Entry.Hash, sizeof(TestHash)) != 0)
-			{
-				UE_LOG(LogPakFile, Error, TEXT("Hash mismatch for \"%s\"."), *It.Filename());
-				ErrorCount++;
-			}
-			else
-			{
-				UE_LOG(LogPakFile, Verbose, TEXT("\"%s\" OK. [%s]"), *It.Filename(), *Info.GetCompressionMethod(Entry.CompressionMethodIndex).ToString());
-			}
-			FMemory::Free(FileContents);
+		FMemoryReader SecondaryIndexReader(FullDirectoryIndexData);
+		{
+			SCOPED_BOOT_TIMING("PakFile_SerializeDirectoryIndex");
+			SecondaryIndexReader << DirectoryIndex;
 		}
-		if (ErrorCount == 0)
+		bHasFullDirectoryIndex = true;
+
+#if ENABLE_PAKFILE_RUNTIME_PRUNING
+		if (bWillUseFullDirectoryIndex)
 		{
-			UE_LOG(LogPakFile, Display, TEXT("Pak file \"%s\" healthy, %d files checked."), *PakFilename, FileCount);
+			bWillPruneDirectoryIndex = false;
 		}
 		else
 		{
-			UE_LOG(LogPakFile, Display, TEXT("Pak file \"%s\" corrupted (%d errors out of %d files checked.)."), *PakFilename, ErrorCount, FileCount);
+			// Store the PrunedDirectoryIndex from the PrimaryIndex in our PrunedDirectoryIndex, to be used for verification and to be swapped into DirectoryIndex later
+			check(bWillUsePathHashIndex); // Need to confirm that we have read the PathHashIndex bytes
+			{
+				SCOPED_BOOT_TIMING("PakFile_SerializePrunedDirectoryIndex");
+				PathHashIndexReader << PrunedDirectoryIndex;
+			}
+			bWillPruneDirectoryIndex = true;
+			bSomePakNeedsPruning = true;
 		}
+#endif
 	}
 
-	double EndTime = FPlatformTime::Seconds();
-	double ElapsedTime = EndTime - StartTime;
-	UE_LOG(LogPakFile, Display, TEXT("Pak file \"%s\" checked in %.2fs"), *PakFilename, ElapsedTime);
+	UE_LOG(LogPakFile, Log, TEXT("PakFile PrimaryIndexSize=%d"), Info.IndexSize);
+	UE_LOG(LogPakFile, Log, TEXT("PakFile PathHashIndexSize=%d"), PathHashIndexSize);
+	UE_LOG(LogPakFile, Log, TEXT("PakFile FullDirectoryIndexSize=%d"), FullDirectoryIndexSize);
 
-	return ErrorCount == 0;
+	check(bHasFullDirectoryIndex || bHasPathHashIndex);
+	return true;
 }
 
-struct FMiniFileEntry
+bool FPakFile::LoadLegacyIndex(FArchive* Reader)
 {
-	uint64 FilenameHash;
-	int32 EntryIndex;
-};
+	bHasPathHashIndex = false;
+	bHasFullDirectoryIndex = false;
+#if ENABLE_PAKFILE_RUNTIME_PRUNING
+	bNeedsLegacyPruning = false;
+	bWillPruneDirectoryIndex = false;
+#endif
 
-static inline int32 CDECL CompareFMiniFileEntry(const void* Left, const void* Right)
-{
-	const FMiniFileEntry* LeftEntry = (const FMiniFileEntry*)Left;
-	const FMiniFileEntry* RightEntry = (const FMiniFileEntry*)Right;
-	if (LeftEntry->FilenameHash < RightEntry->FilenameHash)
-	{
-		return -1;
-	}
-	if (LeftEntry->FilenameHash > RightEntry->FilenameHash)
-	{
-		return 1;
-	}
-	return 0;
-}
+	// Load index into memory first.
 
-bool FPakFile::UnloadPakEntryFilenames(TMap<uint64, FPakEntry>& CrossPakCollisionChecker, TArray<FString>* DirectoryRootsToKeep, bool bAllowRetries)
-{
-	// If the process has already been done, get out of here.
-	if (bAttemptedPakFilenameUnload || bFilenamesRemoved)
+	TArray<uint8> IndexData;
+	if (CachedTotalSize < (Info.IndexOffset + Info.IndexSize))
 	{
-		return true;
+		UE_LOG(LogPakFile, Fatal, TEXT("Corrupted index offset in pak file."));
+		return false;
 	}
 
-	UE_LOG(LogPakFile, Log, TEXT("Unloading filenames for pak '%s'"), *PakFilename);
+	Reader->Seek(Info.IndexOffset);
+	IndexData.SetNum(Info.IndexSize);
+	Reader->Serialize(IndexData.GetData(), Info.IndexSize);
 
-	LLM_SCOPE(ELLMTag::FileSystem);
-
-	// Set this flag so if unloading fails, we don't try again
-	bAttemptedPakFilenameUnload = true;
-
-	if (Info.bIndexIsFrozen)
+	FSHAHash ComputedHash;
+	if (!DecryptAndValidateIndex(Reader, IndexData, Info.IndexHash, ComputedHash))
 	{
-		UE_LOG(LogPakFile, Warning, TEXT("FAILED unloading filenames for pak '%s' - its index was frozen and cannot be modified"), *PakFilename);
+		UE_LOG(LogPakFile, Log, TEXT("Corrupt pak index detected!"));
+		UE_LOG(LogPakFile, Log, TEXT(" Filename: %s"), *PakFilename);
+		UE_LOG(LogPakFile, Log, TEXT(" Encrypted: %d"), Info.bEncryptedIndex);
+		UE_LOG(LogPakFile, Log, TEXT(" Total Size: %d"), Reader->TotalSize());
+		UE_LOG(LogPakFile, Log, TEXT(" Index Offset: %d"), Info.IndexOffset);
+		UE_LOG(LogPakFile, Log, TEXT(" Index Size: %d"), Info.IndexSize);
+		UE_LOG(LogPakFile, Log, TEXT(" Stored Index Hash: %s"), *Info.IndexHash.ToString());
+		UE_LOG(LogPakFile, Log, TEXT(" Computed Index Hash: %s"), *ComputedHash.ToString());
 		return false;
 	}
 
-	// Variables for the filename hashing and collision detection.
-	int NumRetries = 0;
-	const int MAX_RETRIES = bAllowRetries ? 10 : 1;
-	bool bHasCollision;
-	FilenameStartHash = FCrc::StrCrc32(*GetFilename());
-    
-	// Allocate the temporary array for hashing filenames. The Memset is to hopefully
-	// silence the Visual Studio static analyzer.
-	TArray<FMiniFileEntry> MiniFileEntries;
-	MiniFileEntries.AddUninitialized(NumEntries);
 
-	do
+	FMemoryReader IndexReader(IndexData);
+
+	// Read the default mount point and all entries.
+	NumEntries = 0;
+	IndexReader << MountPoint;
+	IndexReader << NumEntries;
+
+	MakeDirectoryFromPath(MountPoint);
+
+	FPakEntryPair PakEntryPair;
+	auto ReadNextEntry = [&PakEntryPair, &IndexReader, this]() -> FPakEntryPair&
 	{
-		// No collisions yet for this pass.
-		bHasCollision = false;
+		IndexReader << PakEntryPair.Filename;
+		PakEntryPair.Info.Serialize(IndexReader, this->Info.Version);
+		return PakEntryPair;
+	};
 
-		TMap<uint64, FPakEntry> NewCollisionCheckEntries;
-		NewCollisionCheckEntries.Reserve(NumEntries);
+	TMap<uint64, FString> CollisionDetection;
+	int32 NumEncodedEntries = 0;
+	int32 NumDeletedEntries = 0;
+#if ENABLE_PAKFILE_RUNTIME_PRUNING
+	bool bCreatePathHash = !IsPakKeepFullDirectory();
+#else
+	// Pruning of legacy files is no longer supported; we will keep the entire directory with no way to prune it.  There is no need to create the PathHashIndex since we will have the FullDirectoryIndex.
+	bool bCreatePathHash = false;
+#endif
+	FPathHashIndex* PathHashToWrite = bCreatePathHash ? &PathHashIndex : nullptr;
+	FPakFile::EncodePakEntriesIntoIndex(NumEntries, ReadNextEntry, *PakFilename, Info, MountPoint, NumEncodedEntries, NumDeletedEntries, &PathHashSeed,
+		&DirectoryIndex, PathHashToWrite, EncodedPakEntries, Files, &CollisionDetection);
+	check(NumEncodedEntries + Files.Num() + NumDeletedEntries == NumEntries);
+	Files.Shrink();
+	EncodedPakEntries.Shrink();
+
+	bHasPathHashIndex = bCreatePathHash;
+	bHasFullDirectoryIndex = true;
+#if ENABLE_PAKFILE_RUNTIME_PRUNING
+	if (!IsPakKeepFullDirectory())
+	{
+		bNeedsLegacyPruning = true;
+		bWillPruneDirectoryIndex = true;
+		bSomePakNeedsPruning = true;
+		// We cannot prune during this call because config settings have not yet been loaded and we need the settings for DirectoryIndexKeepFiles before we can prune
+		// PrunedDirectoryIndex will be created and swapped with DirectoryIndex in OptimizeMemoryUsageForMountedPaks, and bHasFullDirectoryIndex will be set to false then
+	}
+	else
+	{
+		bNeedsLegacyPruning = false;
+		bWillPruneDirectoryIndex = false;
+	}
+#endif
 
-		// Build the list of hashes from the Index based on the starting hash.
-		int32 EntryIndex = 0;
-		FString FinalFilename;
-		FinalFilename.Reset(1024);
-		for (TMemoryImageMap<FMemoryImageString, FPakDirectory>::TConstIterator It(Data->Index); (It && !bHasCollision); ++It)
-		{
-			for (FPakDirectory::TConstIterator DirectoryIt(It.Value()); DirectoryIt; ++DirectoryIt)
-			{
-				FinalFilename.Reset();
-				FinalFilename.Append(It.Key());
-				FinalFilename.PathAppend(*DirectoryIt.Key(), DirectoryIt.Key().Len());
-				FinalFilename.ToLowerInline();
-				uint64 FilenameHash = FFnv::MemFnv64(*FinalFilename, FinalFilename.Len() * sizeof(TCHAR), FilenameStartHash);
-				MiniFileEntries[EntryIndex].FilenameHash = FilenameHash;
-				MiniFileEntries[EntryIndex].EntryIndex = DirectoryIt.Value();
-				++EntryIndex;
+	check(bHasFullDirectoryIndex || bHasPathHashIndex);
+	return true;
+}
 
-				const FPakEntry* EntryFromPreviousPaks = CrossPakCollisionChecker.Find(FilenameHash);
-				const FPakEntry* EntryFromCurrentPak = NewCollisionCheckEntries.Find(FilenameHash);
-				const FPakEntry& CurrentEntry = Data->Files[DirectoryIt.Value()];
+bool FPakFile::DecryptAndValidateIndex(FArchive* Reader, TArray<uint8>& IndexData, FSHAHash& InExpectedHash, FSHAHash& OutActualHash)
+{
+	// Decrypt if necessary
+	if (Info.bEncryptedIndex)
+	{
+		DecryptData(IndexData.GetData(), IndexData.Num(), Info.EncryptionKeyGuid);
+	}
 
-				if (EntryFromPreviousPaks && (FMemory::Memcmp(EntryFromPreviousPaks->Hash, CurrentEntry.Hash, sizeof(CurrentEntry.Hash))) != 0)
-				{
-					UE_LOG(LogPakFile, Verbose, TEXT("Detected collision with previous pak while hashing %s"), *FinalFilename);
-					bHasCollision = true;
-					break;
-				}
+	// Check SHA1 value.
+	FSHA1::HashBuffer(IndexData.GetData(), IndexData.Num(), OutActualHash.Hash);
+	return InExpectedHash == OutActualHash;
+}
 
-				if (EntryFromCurrentPak && (FMemory::Memcmp(EntryFromCurrentPak->Hash, CurrentEntry.Hash, sizeof(CurrentEntry.Hash)) != 0))
-				{
-					UE_LOG(LogPakFile, Verbose, TEXT("Detected collision within pak while hashing %s"), *FinalFilename);
-					bHasCollision = true;
-					break;
-				}
+uint64 FPakFile::HashPath(const TCHAR* RelativePathFromMount, uint64 Seed)
+{
+	FString LowercaseRelativePath(RelativePathFromMount);
+	LowercaseRelativePath.ToLowerInline();
+	return FFnv::MemFnv64(*LowercaseRelativePath, LowercaseRelativePath.Len() * sizeof(TCHAR), Seed);
+}
 
-				NewCollisionCheckEntries.Add(FilenameHash, CurrentEntry);
-			}
+void FPakFile::EncodePakEntriesIntoIndex(int32 InNumEntries, const ReadNextEntryFunction& InReadNextEntry, const TCHAR* InPakFilename, const FPakInfo& InPakInfo, const FString& MountPoint,
+	int32& OutNumEncodedEntries, int32& OutNumDeletedEntries, uint64* OutPathHashSeed,
+	FDirectoryIndex* OutDirectoryIndex, FPathHashIndex* OutPathHashIndex, TArray<uint8>& OutEncodedPakEntries, TArray<FPakEntry>& OutNonEncodableEntries,
+	TMap<uint64, FString>* InOutCollisionDetection)
+{
+	uint64 PathHashSeed = 0;
+	if (OutPathHashSeed || OutPathHashIndex)
+	{
+		PathHashSeed = FCrc::StrCrc32(InPakFilename);
+		if (OutPathHashSeed)
+		{
+			*OutPathHashSeed = PathHashSeed;
 		}
+	}
 
-		if (bHasCollision)
+	OutNumEncodedEntries = 0;
+	OutNumDeletedEntries = 0;
+	FMemoryWriter CompressedEntryWriter(OutEncodedPakEntries);
+	for (int32 EntryCount = 0; EntryCount < InNumEntries; ++EntryCount)
+	{
+		FPakEntryPair& Pair = InReadNextEntry();
+		// Add the Entry and get an FPakEntryLocation for it
+		const FPakEntry& PakEntry = Pair.Info;
+		FPakEntryLocation EntryLocation;
+		if (!PakEntry.IsDeleteRecord())
 		{
-			++NumRetries;
-			++FilenameStartHash;
-			UE_LOG(LogPakFile, Verbose, TEXT("Collisions detected. Retrying with new seed..."));
+			EntryLocation = FPakEntryLocation::CreateFromOffsetIntoEncoded(OutEncodedPakEntries.Num());
+			if (EncodePakEntry(CompressedEntryWriter, PakEntry, InPakInfo))
+			{
+				++OutNumEncodedEntries;
+			}
+			else
+			{
+				int32 ListIndex = OutNonEncodableEntries.Num();
+				EntryLocation = FPakEntryLocation::CreateFromListIndex(ListIndex);
+				OutNonEncodableEntries.Add(PakEntry);
+
+				// PakEntries in the index have some values that are different from the in-place pakentries stored next to each file's payload.  EncodePakEntry handles that internally if encoding succeeded.
+				FPakEntry& StoredPakEntry = OutNonEncodableEntries[ListIndex];
+				FMemory::Memset(StoredPakEntry.Hash, 0); // Hash is 0-filled
+				StoredPakEntry.Verified = true; // Validation of the hash is impossible, so Verified is set to true
+			}
 		}
 		else
 		{
-			CrossPakCollisionChecker.Append(NewCollisionCheckEntries);
+			++OutNumDeletedEntries;
 		}
 
-		// Sort the list to make hash collision detection easy.
-		qsort(MiniFileEntries.GetData(), NumEntries, sizeof(FMiniFileEntry), CompareFMiniFileEntry);
+		// Add the Entry into the requested Indexes
+		AddEntryToIndex(Pair.Filename, EntryLocation, MountPoint, PathHashSeed, OutDirectoryIndex, OutPathHashIndex, InOutCollisionDetection);
 	}
-	while (bHasCollision && NumRetries < MAX_RETRIES);
+}
+
+void FPakFile::PruneDirectoryIndex(FDirectoryIndex& InOutDirectoryIndex, FDirectoryIndex* PrunedDirectoryIndex, const FString& MountPoint)
+{
+	TArray<FString> FileWildCards, DirectoryWildCards, OldWildCards;
+	GConfig->GetArray(TEXT("Pak"), TEXT("DirectoryIndexKeepFiles"), FileWildCards, GEngineIni);
+	GConfig->GetArray(TEXT("Pak"), TEXT("DirectoryIndexKeepEmptyDirectories"), DirectoryWildCards, GEngineIni);
+	GConfig->GetArray(TEXT("Pak"), TEXT("DirectoryRootsToKeepInMemoryWhenUnloadingPakEntryFilenames"), OldWildCards, GEngineIni); // Legacy name, treated as both KeepFiles and KeepEmptyDirectories
+	DirectoryWildCards.Append(OldWildCards);
+	FileWildCards.Append(OldWildCards);
+	int32 NumKeptEntries = 0;
 
-	// Filenames can only be unloaded if we found a collision-free starting hash
-	// within the maximum number of retries.
-	if (NumRetries >= MAX_RETRIES)
+	if (PrunedDirectoryIndex)
 	{
-		//		FPlatformMisc::LowLevelOutputDebugString(TEXT("Can't unload pak filenames due to hash collision..."));
-		UE_LOG(LogPakFile, Warning, TEXT("FAILED unloading filenames for pak '%s'"), *PakFilename);
-		return false;
+		PrunedDirectoryIndex->Empty();
 	}
 
-	// Allocate the storage space.
-	FilenameHashesIndices = new int32[NumEntries];
-	FilenameHashes = new uint64[NumEntries];
-	int32 LastHashMostSignificantBits = -1;
+	TMap<FString, bool> KeepDirectory;
 
-	// FilenameHashesIndex provides small 'arenas' of binary searchable filename hashes.
-	// The most significant bits (MSB) of the hash, 8 in this case, are used to index into
-	// the FilenameHashesIndex to get the start and end indices within FilenameHashes for the
-	// search.
-	//
-	// An example array looks like this:
-	//
-	//     0   - 0       << No entries in the 0-1 MSB range.
-	//     1   - 0       << Entry index 0 begins the 1-2 MSB range.
-	//     2   - 103     << Entry index 103 begins the 2-3 MSB range. The 3 MSB range is 103 also, so there are no entries.
-	//     3   - 103
-	//     4   - 331
-	//     5   - 629
-	//     ...
-	//     256 - 55331   << A value representing NumEntries
-	const int MAX_FILENAME_HASHES_INDEX_SIZE = 257;
-	FilenameHashesIndex = new uint32[MAX_FILENAME_HASHES_INDEX_SIZE];
-
-	// Transfer the sorted hashes to FilenameHashes.
-	for (int32 EntryIndex = 0; EntryIndex < NumEntries; EntryIndex++)
+	// Clear out those portions of the Index allowed by the user.
+	if (DirectoryWildCards.Num() > 0 || FileWildCards.Num() > 0)
 	{
-		// If a new index entry is needed as a result of crossing over into a larger hash group
-		// as specified through the 8 most significant bits of the hash, store the entry index.
-		uint64 FilenameHash = MiniFileEntries[EntryIndex].FilenameHash;
-		int32 HashMostSignificantBits = FilenameHash >> 56;
-		if (HashMostSignificantBits != LastHashMostSignificantBits)
+		for (auto It = InOutDirectoryIndex.CreateIterator(); It; ++It)
 		{
-			for (int32 BitsIndex = LastHashMostSignificantBits + 1; BitsIndex <= HashMostSignificantBits; ++BitsIndex)
+			const FString& DirectoryName = It.Key();
+			FPakDirectory& OriginalDirectory = It.Value();
+			const FString FullDirectoryName = PakPathCombine(MountPoint, DirectoryName);
+			check(IsPathInDirectoryFormat(FullDirectoryName));
+			FPakDirectory* PrunedDirectory = nullptr;
+			bool bKeepDirectory = false;
+
+			TArray<FString> RemoveFilenames;
+			for (auto FileIt = It->Value.CreateIterator(); FileIt; ++FileIt)
 			{
-				FilenameHashesIndex[BitsIndex] = EntryIndex;
-			}
-			LastHashMostSignificantBits = HashMostSignificantBits;
-		}
+				const FString& FileNameWithoutPath = FileIt->Key;
+				const FString FullFilename = PakPathCombine(FullDirectoryName, FileNameWithoutPath);
+				bool bKeepFile = false;
 
-		FilenameHashes[EntryIndex] = FilenameHash;
-		FilenameHashesIndices[EntryIndex] = MiniFileEntries[EntryIndex].EntryIndex;
-	}
+				for (const FString& WildCard : FileWildCards)
+				{
+					if (FullFilename.MatchesWildcard(WildCard))
+					{
+						bKeepFile = true;
+						break;
+					}
+				}
 
-	// Fill out the array to the end.
-	for (int32 BitsIndex = LastHashMostSignificantBits + 1; BitsIndex < MAX_FILENAME_HASHES_INDEX_SIZE; ++BitsIndex)
-	{
-		FilenameHashesIndex[BitsIndex] = NumEntries;
-	}
+				if (bKeepFile)
+				{
+					bKeepDirectory = true;
+					if (PrunedDirectoryIndex)
+					{
+						if (!PrunedDirectory)
+						{
+							PrunedDirectory = &PrunedDirectoryIndex->Add(DirectoryName);
+						}
+						PrunedDirectory->Add(FileNameWithoutPath, FileIt->Value);
+					}
+				}
+				else
+				{
+					if (!PrunedDirectoryIndex)
+					{
+						RemoveFilenames.Add(FileNameWithoutPath);
+					}
+				}
+			}
+			if (!PrunedDirectoryIndex)
+			{
+				for (const FString& FileNameWithoutPath : RemoveFilenames)
+				{
+					OriginalDirectory.Remove(FileNameWithoutPath);
+				}
+			}
 
-	bFilenamesRemoved = true;
+			if (!bKeepDirectory)
+			{
+				for (const FString& WildCard : DirectoryWildCards)
+				{
+					if (FullDirectoryName.MatchesWildcard(WildCard))
+					{
+						bKeepDirectory = true;
+						break;
+					}
+				}
+			}
+			KeepDirectory.FindOrAdd(DirectoryName) = bKeepDirectory;
+		}
 
-#if defined(FPAKFILE_UNLOADPAKENTRYFILENAMES_CHECK)
-	// Build the list of hashes from the Index based on the starting hash.
-	for (TMap<FString, FPakDirectory>::TConstIterator It(Index); It; ++It)
-	{
-		for (FPakDirectory::TConstIterator DirectoryIt(It.Value()); DirectoryIt; ++DirectoryIt)
+		// For each kept directory, mark that we need to keep all of its parents up to the mount point
+		for (const TPair<FString, bool> Pair : KeepDirectory)
 		{
-			int32 EntryIndex = DirectoryIt.Value();
-
-			FString FinalFilename = MountPoint / It.Key() / DirectoryIt.Key();
-			FPakEntry OutEntry;
-			if (!Find(FinalFilename, &OutEntry))
+			if (Pair.Value)
 			{
-				FPlatformMisc::LowLevelOutputDebugString(*FinalFilename);
+				FString CurrentDirectory = Pair.Key;
+				FString UnusedCleanFileName;
+				while (CurrentDirectory != MountPoint)
+				{
+					if (!SplitPathInline(CurrentDirectory, UnusedCleanFileName))
+					{
+						break;
+					}
+					bool& bOldValue = KeepDirectory.FindOrAdd(CurrentDirectory);
+					if (bOldValue)
+					{
+						break;
+					}
+					bOldValue = true;
+				}
 			}
+		}
 
-			FPakEntry& InEntry = Files[EntryIndex];
-			if (InEntry.Offset != OutEntry.Offset ||
-					InEntry.Size != OutEntry.Size ||
-					InEntry.UncompressedSize != OutEntry.UncompressedSize ||
-					InEntry.CompressionMethod != OutEntry.CompressionMethod ||
-					InEntry.bEncrypted != OutEntry.bEncrypted ||
-					InEntry.CompressionBlockSize != OutEntry.CompressionBlockSize ||
-					InEntry.CompressionBlocks != OutEntry.CompressionBlocks)
+		// Prune all of the directories
+		for (const TPair<FString, bool> Pair : KeepDirectory)
+		{
+			const FString& DirectoryName = Pair.Key;
+			bool bKeep = Pair.Value;
+			if (bKeep)
+			{
+				if (PrunedDirectoryIndex)
+				{
+					PrunedDirectoryIndex->FindOrAdd(DirectoryName);
+				}
+			}
+			else
 			{
-				FPlatformMisc::LowLevelOutputDebugString(TEXT("!!!!!!!!!!!!!!!!!!!!!!"));
-				FPlatformMisc::LowLevelOutputDebugString(*FinalFilename);
+				if (!PrunedDirectoryIndex)
+				{
+					InOutDirectoryIndex.Remove(DirectoryName);
+				}
 			}
 		}
 	}
-#endif
-
-	// Clear out those portions of the Index allowed by the user.
-	if (DirectoryRootsToKeep != nullptr)
+	else
 	{
-		for(auto It = Data->Index.CreateIterator(); It; ++It)
+		if (!PrunedDirectoryIndex)
 		{
-			const FString DirectoryName = MountPoint / It.Key();
+			InOutDirectoryIndex.Empty();
+		}
+	}
+}
+
+FPakFile::EFindResult FPakFile::GetPakEntry(const FPakEntryLocation& PakEntryLocation, FPakEntry* OutEntry) const
+{
+	return GetPakEntry(PakEntryLocation, OutEntry, EncodedPakEntries, Files, Info);
+}
 
-			bool bRemoveDirectoryFromIndex = true;
-			for(const FString& DirectoryRoot : *DirectoryRootsToKeep)
+FPakFile::EFindResult FPakFile::GetPakEntry(const FPakEntryLocation& PakEntryLocation, FPakEntry* OutEntry, const TArray<uint8>& EncodedPakEntries, const TArray<FPakEntry>& Files, const FPakInfo& Info)
+{
+	bool bDeleted = PakEntryLocation.IsInvalid();
+	if (OutEntry != NULL)
+	{
+		if (!bDeleted)
+		{
+			// The FPakEntry structures are bit-encoded, so decode it.
+			int32 EncodedOffset = PakEntryLocation.GetAsOffsetIntoEncoded();
+			if (EncodedOffset >= 0)
 			{
-				if (DirectoryName.MatchesWildcard(DirectoryRoot))
-				{
-					bRemoveDirectoryFromIndex = false;
-					break;
-				}
+				check(EncodedOffset < EncodedPakEntries.Num());
+				DecodePakEntry(EncodedPakEntries.GetData() + EncodedOffset, *OutEntry, Info);
 			}
-
-			if (bRemoveDirectoryFromIndex)
+			else
 			{
-				It.RemoveCurrent();
+				int32 ListIndex = PakEntryLocation.GetAsListIndex();
+				check(ListIndex >= 0);
+				const FPakEntry& FoundEntry = Files[ListIndex];
+				//*OutEntry = **FoundEntry;
+				OutEntry->Offset = FoundEntry.Offset;
+				OutEntry->Size = FoundEntry.Size;
+				OutEntry->UncompressedSize = FoundEntry.UncompressedSize;
+				OutEntry->CompressionMethodIndex = FoundEntry.CompressionMethodIndex;
+				OutEntry->CompressionBlocks = FoundEntry.CompressionBlocks;
+				OutEntry->CompressionBlockSize = FoundEntry.CompressionBlockSize;
+				OutEntry->Flags = FoundEntry.Flags;
 			}
 		}
+		else
+		{
+			// entry was deleted, build dummy entry to indicate the deleted entry
+			(*OutEntry) = FPakEntry();
+			OutEntry->SetDeleteRecord(true);
+		}
+
+		// Index PakEntries do not store their hash, so verification against the hash is impossible.
+		// Initialize the OutEntry's Hash to 0 and mark it as already verified.
+		// TODO: Verified and Hash are checked by FPakFileHandle, which is used when opening files from PakFiles synchronously;
+		//       it is not currently used by asynchronous reads in FPakPrecacher, and we can likely remove it from use in FPakFileHandle as well
+		FMemory::Memset(OutEntry->Hash, 0);
+		OutEntry->Verified = true;
+	}
+
+	return bDeleted ? EFindResult::FoundDeleted : EFindResult::Found;
+}
 
-		Data->Index.Shrink();
+struct FPakFile::FIndexSettings
+{
+	FIndexSettings()
+	{
+		bKeepFullDirectory = false;
+		bValidatePruning = false;
+		bDelayPruning = false;
+		bWritePathHashIndex = true;
+		bWriteFullDirectoryIndex = true;
 
-#if defined(FPAKFILE_UNLOADPAKENTRYFILENAMES_LOGKEPTFILENAMES)
-		for (TMap<FString, FPakDirectory>::TConstIterator It(Index); It; ++It)
+		// Paks are mounted before config files are read, so the licensee needs to hardcode all settings used for runtime index loading rather than specifying them in ini
+		if (FPakPlatformFile::GetPakSetIndexSettingsDelegate().IsBound())
 		{
-			FPlatformMisc::LowLevelOutputDebugStringf(TEXT("FPakFile::UnloadPakEntryFilenames() %s - Keeping %s"), *PakFilename, *It.Key());
+			FPakPlatformFile::GetPakSetIndexSettingsDelegate().Execute(bKeepFullDirectory, bValidatePruning, bDelayPruning);
 		}
+
+		// Settings not used at runtime can be read from ini
+#if !UE_BUILD_SHIPPING
+		GConfig->GetBool(TEXT("Pak"), TEXT("WritePathHashIndex"), bWritePathHashIndex, GEngineIni);
+		GConfig->GetBool(TEXT("Pak"), TEXT("WriteFullDirectoryIndex"), bWriteFullDirectoryIndex, GEngineIni);
 #endif
+
+		bKeepFullDirectory = bKeepFullDirectory || IS_PROGRAM || !FPlatformProperties::RequiresCookedData();
+#if !UE_BUILD_SHIPPING
+		const TCHAR* CommandLine = FCommandLine::Get();
+		FParse::Bool(CommandLine, TEXT("ForcePakKeepFullDirectory="), bKeepFullDirectory);
+#if ENABLE_PAKFILE_RUNTIME_PRUNING_VALIDATE
+		FParse::Bool(CommandLine, TEXT("ForcePakValidatePruning="), bValidatePruning);
+		FParse::Bool(CommandLine, TEXT("ForcePakDelayPruning="), bDelayPruning);
+#endif
+		FParse::Bool(CommandLine, TEXT("ForcePakWritePathHashIndex="), bWritePathHashIndex);
+		FParse::Bool(CommandLine, TEXT("ForcePakWriteFullDirectoryIndex="), bWriteFullDirectoryIndex);
+#endif
+	}
+
+	bool bKeepFullDirectory;
+	bool bValidatePruning;
+	bool bDelayPruning;
+	bool bWritePathHashIndex;
+	bool bWriteFullDirectoryIndex;
+};
+
+FPakFile::FIndexSettings& FPakFile::GetIndexSettings()
+{
+	static FIndexSettings IndexLoadParams;
+	return IndexLoadParams;
+}
+
+bool FPakFile::IsPakKeepFullDirectory()
+{
+	FIndexSettings& IndexLoadParams = GetIndexSettings();
+	return IndexLoadParams.bKeepFullDirectory;
+}
+
+bool FPakFile::IsPakValidatePruning()
+{
+	FIndexSettings& IndexLoadParams = GetIndexSettings();
+	return IndexLoadParams.bValidatePruning;
+}
+
+bool FPakFile::IsPakDelayPruning()
+{
+	FIndexSettings& IndexLoadParams = GetIndexSettings();
+	return IndexLoadParams.bDelayPruning;
+}
+
+bool FPakFile::IsPakWritePathHashIndex()
+{
+	FIndexSettings& IndexLoadParams = GetIndexSettings();
+	return IndexLoadParams.bWritePathHashIndex;
+}
+
+bool FPakFile::IsPakWriteFullDirectoryIndex()
+{
+	FIndexSettings& IndexLoadParams = GetIndexSettings();
+	return IndexLoadParams.bWriteFullDirectoryIndex;
+}
+
+bool FPakFile::RequiresDirectoryIndexLock() const
+{
+#if ENABLE_PAKFILE_RUNTIME_PRUNING
+	return bWillPruneDirectoryIndex;
+#else
+	return false; 
+#endif
+}
+
+bool FPakFile::ShouldValidatePrunedDirectory() const
+{
+#if ENABLE_PAKFILE_RUNTIME_PRUNING_VALIDATE
+	return IsPakValidatePruning() && bWillPruneDirectoryIndex && !bNeedsLegacyPruning;
+#else
+	return false;
+#endif
+}
+
+
+void FPakFile::AddEntryToIndex(const FString& Filename, const FPakEntryLocation& EntryLocation, const FString& MountPoint, uint64 PathHashSeed,
+	FDirectoryIndex* DirectoryIndex, FPathHashIndex* PathHashIndex, TMap<uint64, FString>* CollisionDetection)
+{
+	FString RelativePathFromMount;
+	if (FPaths::IsRelative(Filename))
+	{
+		RelativePathFromMount = Filename;
 	}
 	else
 	{
-		Data->Index.Empty(0);
+		check(IsPathInDirectoryFormat(MountPoint));
+		check(Filename.StartsWith(MountPoint));
+		RelativePathFromMount = Filename.Mid(MountPoint.Len());
 	}
 
-	return true;
-}
+	if (DirectoryIndex)
+	{
+		FString RelativeDirectoryFromMount = RelativePathFromMount;
+		FString CleanFileName;
+		SplitPathInline(RelativeDirectoryFromMount, CleanFileName);
+		FPakDirectory* Directory = DirectoryIndex->Find(RelativeDirectoryFromMount);
+		if (Directory == nullptr)
+		{
+			// add the parent directories up to the mount point (mount point relative path is "/")
+			FString CurrentDirectory(RelativeDirectoryFromMount);
+			FString UnusedCleanFileName;
+			while (!CurrentDirectory.IsEmpty())
+			{
+				// FPaths::GetPath doesn't handle our / at the end of directories, so call FindLastChar(/)
+				if (!SplitPathInline(CurrentDirectory, UnusedCleanFileName))
+				{
+					break;
+				}
+				DirectoryIndex->FindOrAdd(CurrentDirectory);
+			}
 
-bool FPakFile::ShrinkPakEntriesMemoryUsage()
-{
-	// If the process has already been done, get out of here.
-	if (bAttemptedPakEntryShrink || MiniPakEntries != nullptr)
+			// Add the new directory; this has to come after the addition of the parent directories because we want to use the pointer afterwards and adding other directories would invalidate it
+			Directory = &DirectoryIndex->Add(RelativeDirectoryFromMount);
+		}
+		Directory->Add(CleanFileName, EntryLocation);
+	}
+
+	// Add the entry into the PathHash index
+	if (CollisionDetection || PathHashIndex)
 	{
-		return true;
+		uint64 PathHash = FPakFile::HashPath(*RelativePathFromMount, PathHashSeed);
+		if (CollisionDetection)
+		{
+			FString* ExistingFilename = CollisionDetection->Find(PathHash);
+			checkf(!ExistingFilename || ExistingFilename->Equals(RelativePathFromMount), TEXT("Hash collision for two Filenames within a PakFile.  Filename1 == '%s'.  Filename2 == '%s'.  Hash='0x%x'.")
+				TEXT(" Collision handling is not yet implemented; to resolve the conflict, modify one of the Filenames."), **ExistingFilename, *RelativePathFromMount, PathHash);
+			CollisionDetection->Add(PathHash, RelativePathFromMount);
+		}
+		if (PathHashIndex)
+		{
+			PathHashIndex->Add(PathHash, EntryLocation);
+		}
 	}
+}
 
-	LLM_SCOPE(ELLMTag::FileSystem);
+// Take a pak entry and byte encode it into a smaller representation
+bool FPakFile::EncodePakEntry(FArchive& Ar, const FPakEntry& InPakEntry, const FPakInfo& InInfo)
+{
+	// See notes in DecodePakEntry for the output layout
 
-	UE_LOG(LogPakFile, Log, TEXT("Shrinking entries for pak '%s'"), *PakFilename);
+	check(Ar.IsSaving()); // This function is encode only, and promises not to modify InEntry, but we want to use << which takes non-const so we need to assert that Ar is a loader not a saver
+	check(!InPakEntry.IsDeleteRecord()); // Deleted PakEntries cannot be encoded, caller must check for IsDeleteRecord and handle it separately by e.g. not adding the Entry to the FileList and instead giving the referencer an invalid FPakEntryLocation
+	FPakEntry& PakEntry = const_cast<FPakEntry&>(InPakEntry);
 
-	// Set this flag so if shrinking fails, we don't try again
-	bAttemptedPakEntryShrink = true;
+	const uint32 CompressedBlockAlignment = PakEntry.IsEncrypted() ? FAES::AESBlockSize : 1;
+	const int64 HeaderSize = PakEntry.GetSerializedSize(InInfo.Version);
 
-	// if the index was frozen, we can't unload parts of it, so skip on out
-	if (Info.bIndexIsFrozen)
+	// This data fits into a bitfield (described in DecodePakEntry), and the data has
+	// to fit within a certain range of bits.
+	if (PakEntry.CompressionMethodIndex >= (1 << 6))
 	{
-		UE_LOG(LogPakFile, Warning, TEXT("FAILED shrinking entries for pak file '%s' - its index was frozen and cannot be modified"), *PakFilename);
 		return false;
 	}
-
-	// Wander every file entry.
-	int TotalSizeOfCompressedEntries = 0;
-	bool bIsPossibleToShrink = true;
-	int32 EntryIndex = 0;
-	for (EntryIndex = 0; EntryIndex < NumEntries; ++EntryIndex)
+	if (PakEntry.CompressionBlocks.Num() >= (1 << 16))
 	{
-		FPakEntry& Entry = Data->Files[EntryIndex];
-
-		bool bIsOffset32BitSafe = Entry.Offset <= MAX_uint32;
-		bool bIsSize32BitSafe = Entry.Size <= MAX_uint32;
-		bool bIsUncompressedSize32BitSafe = Entry.UncompressedSize <= MAX_uint32;
-		uint32 CompressedBlockAlignment = Entry.IsEncrypted() ? FAES::AESBlockSize : 1;
-		int64 HeaderSize = Entry.GetSerializedSize(Info.Version);
-
-		// This data fits into a bitfield (described below), and the data has
-		// to fit within a certain range of bits.
-		if (Entry.CompressionMethodIndex >= (1 << 6))
+		return false;
+	}
+	if (PakEntry.CompressionMethodIndex != 0)
+	{
+		if (PakEntry.CompressionBlockSize != PakEntry.UncompressedSize && ((PakEntry.CompressionBlockSize >> 11) > 0x3f))
 		{
-			bIsPossibleToShrink = false;
-			break;
+			return false;
 		}
-		if (Entry.CompressionBlocks.Num() >= (1 << 16))
+		if (PakEntry.CompressionBlocks.Num() > 0 && ((InInfo.HasRelativeCompressedChunkOffsets() ? 0 : PakEntry.Offset) + HeaderSize != PakEntry.CompressionBlocks[0].CompressedStart))
 		{
-			bIsPossibleToShrink = false;
-			break;
+			return false;
 		}
-		if (Entry.CompressionMethodIndex != 0)
+		if (PakEntry.CompressionBlocks.Num() == 1)
 		{
-			if (Entry.CompressionBlockSize != Entry.UncompressedSize && ((Entry.CompressionBlockSize >> 11) > 0x3f))
+			uint64 Base = InInfo.HasRelativeCompressedChunkOffsets() ? 0 : PakEntry.Offset;
+			uint64 AlignedBlockSize = Align(PakEntry.CompressionBlocks[0].CompressedEnd - PakEntry.CompressionBlocks[0].CompressedStart, CompressedBlockAlignment);
+			if ((Base + HeaderSize + PakEntry.Size) != (PakEntry.CompressionBlocks[0].CompressedStart + AlignedBlockSize))
 			{
-				bIsPossibleToShrink = false;
-				break;
-			}
-			if (Entry.CompressionBlocks.Num() > 0 && ((Info.HasRelativeCompressedChunkOffsets() ? 0 : Entry.Offset) + HeaderSize != Entry.CompressionBlocks[0].CompressedStart))
-			{
-				bIsPossibleToShrink = false;
-				break;
+				return false;
 			}
-			if (Entry.CompressionBlocks.Num() == 1)
+		}
+		if (PakEntry.CompressionBlocks.Num() > 1)
+		{
+			for (int i = 1; i < PakEntry.CompressionBlocks.Num(); ++i)
 			{
-				uint64 Base = Info.HasRelativeCompressedChunkOffsets() ? 0 : Entry.Offset;
-				uint64 AlignedBlockSize = Align(Entry.CompressionBlocks[0].CompressedEnd - Entry.CompressionBlocks[0].CompressedStart, CompressedBlockAlignment);
-				if ((Base + HeaderSize + Entry.Size) != (Entry.CompressionBlocks[0].CompressedStart + AlignedBlockSize))
+				uint64 PrevBlockSize = PakEntry.CompressionBlocks[i - 1].CompressedEnd - PakEntry.CompressionBlocks[i - 1].CompressedStart;
+				PrevBlockSize = Align(PrevBlockSize, CompressedBlockAlignment);
+				if (PakEntry.CompressionBlocks[i].CompressedStart != (PakEntry.CompressionBlocks[i - 1].CompressedStart + PrevBlockSize))
 				{
-					bIsPossibleToShrink = false;
-					break;
+					return false;
 				}
 			}
-			if (Entry.CompressionBlocks.Num() > 1)
-			{
-				for (int i = 1; i < Entry.CompressionBlocks.Num(); ++i)
-				{
-					uint64 PrevBlockSize = Entry.CompressionBlocks[i - 1].CompressedEnd - Entry.CompressionBlocks[i - 1].CompressedStart;
-					PrevBlockSize = Align(PrevBlockSize, CompressedBlockAlignment);
-					if (Entry.CompressionBlocks[i].CompressedStart != (Entry.CompressionBlocks[i - 1].CompressedStart + PrevBlockSize))
-					{
-						bIsPossibleToShrink = false;
-						break;
-					}
-				}
+		}
+	}
 
-				if (!bIsPossibleToShrink)
-				{
-					break;
-				}
-			}
+	// This entry can be encoded, so let's do it!
+
+	bool bIsOffset32BitSafe = PakEntry.Offset <= MAX_uint32;
+	bool bIsSize32BitSafe = PakEntry.Size <= MAX_uint32;
+	bool bIsUncompressedSize32BitSafe = PakEntry.UncompressedSize <= MAX_uint32;
+
+	// Build the Flags field.
+	uint32 Flags =
+		(bIsOffset32BitSafe ? (1 << 31) : 0)
+		| (bIsUncompressedSize32BitSafe ? (1 << 30) : 0)
+		| (bIsSize32BitSafe ? (1 << 29) : 0)
+		| (PakEntry.CompressionMethodIndex << 23)
+		| (PakEntry.IsEncrypted() ? (1 << 22) : 0)
+		| (PakEntry.CompressionBlocks.Num() << 6)
+		| (PakEntry.CompressionBlockSize >> 11)
+		;
+
+	Ar << Flags;
+
+	// Build the Offset field.
+	if (bIsOffset32BitSafe)
+	{
+		uint32 Value = (uint32)PakEntry.Offset;
+		Ar << Value;
+	}
+	else
+	{
+		Ar << PakEntry.Offset;
+	}
+
+	// Build the Uncompressed Size field.
+	if (bIsUncompressedSize32BitSafe)
+	{
+		uint32 Value = (uint32)PakEntry.UncompressedSize;
+		Ar << Value;
+	}
+	else
+	{
+		Ar << PakEntry.UncompressedSize;
+	}
+
+	// Any additional data is for compressed file data.
+	if (PakEntry.CompressionMethodIndex != 0)
+	{
+		// Build the Compressed Size field.
+		if (bIsSize32BitSafe)
+		{
+			uint32 Value = (uint32)PakEntry.Size;
+			Ar << Value;
+		}
+		else
+		{
+			Ar << PakEntry.Size;
 		}
 
-		TotalSizeOfCompressedEntries += sizeof(uint32)
-			+ (bIsOffset32BitSafe ? sizeof(uint32) : sizeof(uint64))
-			+ (bIsUncompressedSize32BitSafe ? sizeof(uint32) : sizeof(uint64));
-		if (Entry.CompressionMethodIndex != 0)
+		// Build the Compression Blocks array.
+		if (PakEntry.CompressionBlocks.Num() > 1 || (PakEntry.CompressionBlocks.Num() == 1 && PakEntry.IsEncrypted()))
 		{
-			TotalSizeOfCompressedEntries +=
-				(bIsSize32BitSafe ? sizeof(uint32) : sizeof(uint64));
-			if (Entry.CompressionBlocks.Num() > 1 || (Entry.CompressionBlocks.Num() == 1 && Entry.IsEncrypted()))
+			for (int CompressionBlockIndex = 0; CompressionBlockIndex < PakEntry.CompressionBlocks.Num(); ++CompressionBlockIndex)
 			{
-				TotalSizeOfCompressedEntries += Entry.CompressionBlocks.Num() * sizeof(uint32);
+				uint32 Value = PakEntry.CompressionBlocks[CompressionBlockIndex].CompressedEnd - PakEntry.CompressionBlocks[CompressionBlockIndex].CompressedStart;
+				Ar << Value;
 			}
 		}
 	}
 
-	if (!bIsPossibleToShrink)
+	return true;
+}
+
+void FPakFile::DecodePakEntry(const uint8* SourcePtr, FPakEntry& OutEntry, const FPakInfo& InInfo)
+{
+	// Grab the big bitfield value:
+	// Bit 31 = Offset 32-bit safe?
+	// Bit 30 = Uncompressed size 32-bit safe?
+	// Bit 29 = Size 32-bit safe?
+	// Bits 28-23 = Compression method
+	// Bit 22 = Encrypted
+	// Bits 21-6 = Compression blocks count
+	// Bits 5-0 = Compression block size
+	uint32 Value = *(uint32*)SourcePtr;
+	SourcePtr += sizeof(uint32);
+
+	// Filter out the CompressionMethod.
+	OutEntry.CompressionMethodIndex = (Value >> 23) & 0x3f;
+
+	// Test for 32-bit safe values. Grab it, or memcpy the 64-bit value
+	// to avoid alignment exceptions on platforms requiring 64-bit alignment
+	// for 64-bit variables.
+	//
+	// Read the Offset.
+	bool bIsOffset32BitSafe = (Value & (1 << 31)) != 0;
+	if (bIsOffset32BitSafe)
 	{
-		UE_LOG(LogPakFile, Warning, TEXT("FAILED shrinking entries for pak file '%s'"), *PakFilename);
-		return false;
+		OutEntry.Offset = *(uint32*)SourcePtr;
+		SourcePtr += sizeof(uint32);
+	}
+	else
+	{
+		FMemory::Memcpy(&OutEntry.Offset, SourcePtr, sizeof(int64));
+		SourcePtr += sizeof(int64);
+	}
+
+	// Read the UncompressedSize.
+	bool bIsUncompressedSize32BitSafe = (Value & (1 << 30)) != 0;
+	if (bIsUncompressedSize32BitSafe)
+	{
+		OutEntry.UncompressedSize = *(uint32*)SourcePtr;
+		SourcePtr += sizeof(uint32);
+	}
+	else
+	{
+		FMemory::Memcpy(&OutEntry.UncompressedSize, SourcePtr, sizeof(int64));
+		SourcePtr += sizeof(int64);
 	}
 
-	// Allocate the buffer to hold onto all of the bit-encoded compressed FPakEntry structures.
-	MiniPakEntries = new uint8[TotalSizeOfCompressedEntries];
-	MiniPakEntriesOffsets = new uint32[NumEntries];
-
-	// Walk all of the file entries.
-	uint8* CurrentEntryPtr = MiniPakEntries;
-	for (EntryIndex = 0; EntryIndex < NumEntries; ++EntryIndex)
-	{
-		FPakEntry* FullEntry = &Data->Files[EntryIndex];
-
-		MiniPakEntriesOffsets[EntryIndex] = CurrentEntryPtr - MiniPakEntries;
-
-		//deleted records have a magic number in the offset instead (not ideal, but there is no more space in the bit-encoded entry)
-		if (FullEntry->IsDeleteRecord())
-		{
-			MiniPakEntriesOffsets[EntryIndex] = MAX_uint32;
-		}
-
-		// Begin building the compressed memory structure.
-		//
-		// The general data format for a bit-encoded entry is this:
-		//
-		//     uint32 - Flags
-		//                Bit 31 = Offset 32-bit safe?
-		//                Bit 30 = Uncompressed size 32-bit safe?
-		//                Bit 29 = Size 32-bit safe?
-		//                Bits 28-23 = Compression method
-		//                Bit 22 = Encrypted
-		//                Bits 21-6 = Compression blocks count
-		//                Bits 5-0 = Compression block size
-		//     uint32/uint64 - Offset (either 32-bit or 64-bit depending on bIsOffset32BitSafe)
-		//     uint32/uint64 - Uncompressed Size (either 32-bit or 64-bit depending on bIsUncompressedSize32BitSafe)
-		//
-		//   If the CompressionMethod != COMPRESS_None:
-		//     uint32/uint64 - Size (either 32-bit or 64-bit depending on bIsSize32BitSafe)
-		//
-		//     If the Compression blocks count is more than 1, then an array of Compression block sizes follows of:
-		//         uint32    - Number of bytes in this Compression block.
-		//
-		bool bIsOffset32BitSafe = FullEntry->Offset <= MAX_uint32;
-		bool bIsSize32BitSafe = FullEntry->Size <= MAX_uint32;
-		bool bIsUncompressedSize32BitSafe = FullEntry->UncompressedSize <= MAX_uint32;
-
-		// Build the Flags field.
-		*(uint32*)CurrentEntryPtr =
-			(bIsOffset32BitSafe ? (1 << 31) : 0)
-			| (bIsUncompressedSize32BitSafe ? (1 << 30) : 0)
-			| (bIsSize32BitSafe ? (1 << 29) : 0)
-			| (FullEntry->CompressionMethodIndex << 23)
-			| (FullEntry->IsEncrypted() ? (1 << 22) : 0)
-			| (FullEntry->CompressionBlocks.Num() << 6)
-			| (FullEntry->CompressionBlockSize >> 11)
-			;
-		CurrentEntryPtr += sizeof(uint32);
-
-		// Build the Offset field.
-		if (bIsOffset32BitSafe)
-		{
-			*(uint32*)CurrentEntryPtr = (uint32)FullEntry->Offset;
-			CurrentEntryPtr += sizeof(uint32);
+	// Fill in the Size.
+	if (OutEntry.CompressionMethodIndex != 0)
+	{
+		// Size is only present if compression is applied.
+		bool bIsSize32BitSafe = (Value & (1 << 29)) != 0;
+		if (bIsSize32BitSafe)
+		{
+			OutEntry.Size = *(uint32*)SourcePtr;
+			SourcePtr += sizeof(uint32);
 		}
 		else
 		{
-			FMemory::Memcpy(CurrentEntryPtr, &FullEntry->Offset, sizeof(int64));
-			CurrentEntryPtr += sizeof(int64);
+			FMemory::Memcpy(&OutEntry.Size, SourcePtr, sizeof(int64));
+			SourcePtr += sizeof(int64);
 		}
+	}
+	else
+	{
+		// The Size is the same thing as the UncompressedSize when
+		// CompressionMethod == COMPRESS_None.
+		OutEntry.Size = OutEntry.UncompressedSize;
+	}
+
+	// Filter the encrypted flag.
+	OutEntry.SetEncrypted((Value & (1 << 22)) != 0);
+
+	// This should clear out any excess CompressionBlocks that may be valid in the user's
+	// passed in entry.
+	uint32 CompressionBlocksCount = (Value >> 6) & 0xffff;
+	OutEntry.CompressionBlocks.Empty(CompressionBlocksCount);
+	OutEntry.CompressionBlocks.SetNum(CompressionBlocksCount);
+
+	// Filter the compression block size or use the UncompressedSize if less that 64k.
+	OutEntry.CompressionBlockSize = 0;
+	if (CompressionBlocksCount > 0)
+	{
+		OutEntry.CompressionBlockSize = OutEntry.UncompressedSize < 65536 ? (uint32)OutEntry.UncompressedSize : ((Value & 0x3f) << 11);
+	}
+
+	// Set bDeleteRecord to false, because it obviously isn't deleted if we are here.
+	OutEntry.SetDeleteRecord(false);
+
+	// Base offset to the compressed data
+	int64 BaseOffset = InInfo.HasRelativeCompressedChunkOffsets() ? 0 : OutEntry.Offset;
 
-		// Build the Uncompressed Size field.
-		if (bIsUncompressedSize32BitSafe)
+	// Handle building of the CompressionBlocks array.
+	if (OutEntry.CompressionBlocks.Num() == 1 && !OutEntry.IsEncrypted())
+	{
+		// If the number of CompressionBlocks is 1, we didn't store any extra information.
+		// Derive what we can from the entry's file offset and size.
+		FPakCompressedBlock& CompressedBlock = OutEntry.CompressionBlocks[0];
+		CompressedBlock.CompressedStart = BaseOffset + OutEntry.GetSerializedSize(InInfo.Version);
+		CompressedBlock.CompressedEnd = CompressedBlock.CompressedStart + OutEntry.Size;
+	}
+	else if (OutEntry.CompressionBlocks.Num() > 0)
+	{
+		// Get the right pointer to start copying the CompressionBlocks information from.
+		uint32* CompressionBlockSizePtr = (uint32*)SourcePtr;
+
+		// Alignment of the compressed blocks
+		uint64 CompressedBlockAlignment = OutEntry.IsEncrypted() ? FAES::AESBlockSize : 1;
+
+		// CompressedBlockOffset is the starting offset. Everything else can be derived from there.
+		int64 CompressedBlockOffset = BaseOffset + OutEntry.GetSerializedSize(InInfo.Version);
+		for (int CompressionBlockIndex = 0; CompressionBlockIndex < OutEntry.CompressionBlocks.Num(); ++CompressionBlockIndex)
 		{
-			*(uint32*)CurrentEntryPtr = (uint32)FullEntry->UncompressedSize;
-			CurrentEntryPtr += sizeof(uint32);
+			FPakCompressedBlock& CompressedBlock = OutEntry.CompressionBlocks[CompressionBlockIndex];
+			CompressedBlock.CompressedStart = CompressedBlockOffset;
+			CompressedBlock.CompressedEnd = CompressedBlockOffset + *CompressionBlockSizePtr++;
+			CompressedBlockOffset += Align(CompressedBlock.CompressedEnd - CompressedBlock.CompressedStart, CompressedBlockAlignment);
 		}
-		else
+	}
+}
+
+bool FPakFile::NormalizeDirectoryQuery(const TCHAR* InPath, FString& OutRelativePathFromMount) const
+{
+	OutRelativePathFromMount = InPath;
+	MakeDirectoryFromPath(OutRelativePathFromMount);
+
+	// Check the specified path is under the mount point of this pak file.
+	if (!OutRelativePathFromMount.StartsWith(MountPoint))
+	{
+		return false;
+	}
+	OutRelativePathFromMount = OutRelativePathFromMount.Mid(MountPoint.Len());
+	return true;
+}
+
+const FPakDirectory* FPakFile::FindPrunedDirectoryInternal(const FString& RelativePathFromMount) const
+{
+	const FPakDirectory* PakDirectory = nullptr;
+
+#if ENABLE_PAKFILE_RUNTIME_PRUNING_VALIDATE
+	if (ShouldValidatePrunedDirectory())
+	{
+		PakDirectory = DirectoryIndex.Find(RelativePathFromMount);
+		const FPakDirectory* PrunedPakDirectory = PrunedDirectoryIndex.Find(RelativePathFromMount);
+		if ((PakDirectory != nullptr) != (PrunedPakDirectory != nullptr))
 		{
-			FMemory::Memcpy(CurrentEntryPtr, &FullEntry->UncompressedSize, sizeof(int64));
-			CurrentEntryPtr += sizeof(int64);
+			TSet<FString> FullFoundFiles, PrunedFoundFiles;
+			FString ReportedDirectoryName = MountPoint + RelativePathFromMount;
+			if (PakDirectory)
+			{
+				FullFoundFiles.Add(ReportedDirectoryName);
+			}
+			if (PrunedPakDirectory)
+			{
+				PrunedFoundFiles.Add(ReportedDirectoryName);
+			}
+			ValidateDirectorySearch(FullFoundFiles, PrunedFoundFiles, *ReportedDirectoryName);
 		}
+	}
+	else
+#endif
+	{
+		PakDirectory = DirectoryIndex.Find(RelativePathFromMount);
+	}
+	return PakDirectory;
+}
+
+bool FPakFile::Check()
+{
+	UE_LOG(LogPakFile, Display, TEXT("Checking pak file \"%s\". This may take a while..."), *PakFilename);
+	double StartTime = FPlatformTime::Seconds();
+
+	FArchive& PakReader = *GetSharedReader(NULL);
+	int32 ErrorCount = 0;
+	int32 FileCount = 0;
+
+	bool bSuccess = true;
+
+	// If the pak file is signed, we can do a fast check by just reading a single byte from the start of
+	// each signing block. The signed archive reader will bring in that whole signing block and compare
+	// against the signature table and fire the handler
+	if (bSigned)
+	{
+		FDelegateHandle DelegateHandle;
+		FPakPlatformFile::FPakSigningFailureHandlerData& HandlerData = FPakPlatformFile::GetPakSigningFailureHandlerData();
 
-		// Any additional data is for compressed file data.
-		if (FullEntry->CompressionMethodIndex != 0)
 		{
-			// Build the Compressed Size field.
-			if (bIsSize32BitSafe)
+			FScopeLock Lock(&HandlerData.Lock);
+			DelegateHandle = HandlerData.ChunkSignatureCheckFailedDelegate.AddLambda([&bSuccess](const FPakChunkSignatureCheckFailedData&)
 			{
-				*(uint32*)CurrentEntryPtr = (uint32)FullEntry->Size;
-				CurrentEntryPtr += sizeof(uint32);
+				bSuccess = false;
+			});
+		}
+
+		int64 CurrentPos = 0;
+		const int64 Size = PakReader.TotalSize();
+		while (CurrentPos < Size)
+		{
+			PakReader.Seek(CurrentPos);
+			uint8 Byte = 0;
+			PakReader << Byte;
+			CurrentPos += FPakInfo::MaxChunkDataSize;
+		}
+
+		if (DelegateHandle.IsValid())
+		{
+			FScopeLock Lock(&HandlerData.Lock);
+			HandlerData.ChunkSignatureCheckFailedDelegate.Remove(DelegateHandle);
+		}
+	}
+	else
+	{
+		const bool bIncludeDeleted = true;
+		TCHAR EntryNameBuffer[256];
+		auto GetEntryName = [&EntryNameBuffer](const FPakFile::FPakEntryIterator& It)
+		{
+			const FString* EntryFilename = It.TryGetFilename();
+			if (EntryFilename)
+			{
+				TCString<TCHAR>::Snprintf(EntryNameBuffer, sizeof(EntryNameBuffer), TEXT("\"%s\""), **EntryFilename);
 			}
 			else
 			{
-				FMemory::Memcpy(CurrentEntryPtr, &FullEntry->Size, sizeof(int64));
-				CurrentEntryPtr += sizeof(int64);
+				TCString<TCHAR>::Snprintf(EntryNameBuffer, sizeof(EntryNameBuffer), TEXT("file at offset %u"), It.Info().Offset);
+			}
+			return EntryNameBuffer;
+		};
+		for (FPakFile::FPakEntryIterator It(*this, bIncludeDeleted); It; ++It, ++FileCount)
+		{
+			const FPakEntry& EntryFromIndex = It.Info();
+			if (EntryFromIndex.IsDeleteRecord())
+			{
+				UE_LOG(LogPakFile, Verbose, TEXT("%s Deleted."), GetEntryName(It));
+				continue;
 			}
 
-			// Build the Compression Blocks array.
-			if (FullEntry->CompressionBlocks.Num() > 1 || (FullEntry->CompressionBlocks.Num() == 1 && FullEntry->IsEncrypted()))
+			void* FileContents = FMemory::Malloc(EntryFromIndex.Size);
+			PakReader.Seek(EntryFromIndex.Offset);
+			uint32 SerializedCrcTest = 0;
+			FPakEntry EntryFromPayload;
+			EntryFromPayload.Serialize(PakReader, GetInfo().Version);
+			if (!EntryFromPayload.IndexDataEquals(EntryFromIndex))
 			{
-				for (int CompressionBlockIndex = 0; CompressionBlockIndex < FullEntry->CompressionBlocks.Num(); ++CompressionBlockIndex)
-				{
-					*(uint32*)CurrentEntryPtr = FullEntry->CompressionBlocks[CompressionBlockIndex].CompressedEnd - FullEntry->CompressionBlocks[CompressionBlockIndex].CompressedStart;
-					CurrentEntryPtr += sizeof(uint32);
-				}
+				UE_LOG(LogPakFile, Error, TEXT("Index FPakEntry does not match Payload FPakEntry for %s."), GetEntryName(It));
+				ErrorCount++;
 			}
-		}
+			PakReader.Serialize(FileContents, EntryFromIndex.Size);
 
-#if !UE_BUILD_SHIPPING
-		if (!FullEntry->IsDeleteRecord())
+			uint8 TestHash[20];
+			FSHA1::HashBuffer(FileContents, EntryFromIndex.Size, TestHash);
+			if (FMemory::Memcmp(TestHash, EntryFromPayload.Hash, sizeof(TestHash)) != 0)
+			{
+				UE_LOG(LogPakFile, Error, TEXT("Hash mismatch for %s."), GetEntryName(It));
+				ErrorCount++;
+			}
+			else
+			{
+				UE_LOG(LogPakFile, Verbose, TEXT("%s OK. [%s]"), GetEntryName(It), *Info.GetCompressionMethod(EntryFromIndex.CompressionMethodIndex).ToString());
+			}
+			FMemory::Free(FileContents);
+		}
+		if (ErrorCount == 0)
+		{
+			UE_LOG(LogPakFile, Display, TEXT("Pak file \"%s\" healthy, %d files checked."), *PakFilename, FileCount);
+		}
+		else
 		{
-			FPakEntry Test;
-			DecodePakEntry(MiniPakEntries + MiniPakEntriesOffsets[EntryIndex], &Test);
-			FMemory::Memcpy(Test.Hash, FullEntry->Hash, 20);
-			check(Test == *FullEntry);
+			UE_LOG(LogPakFile, Display, TEXT("Pak file \"%s\" corrupted (%d errors out of %d files checked.)."), *PakFilename, ErrorCount, FileCount);
 		}
-#endif
 	}
 
-	check(CurrentEntryPtr == MiniPakEntries + TotalSizeOfCompressedEntries);
-
-	// Clear out the Files data. We compressed it, and we don't need the wasted
-	// space of the original anymore.
-	Data->Files.Empty(0);
-
-	static uint64 Total = 0;
-	Total += TotalSizeOfCompressedEntries;
-	UE_LOG(LogPakFile, Display, TEXT("Compressed pak entries down to %d bytes [Total = %llu bytes]"), TotalSizeOfCompressedEntries, Total);
+	double EndTime = FPlatformTime::Seconds();
+	double ElapsedTime = EndTime - StartTime;
+	UE_LOG(LogPakFile, Display, TEXT("Pak file \"%s\" checked in %.2fs"), *PakFilename, ElapsedTime);
 
-	return true;
+	return ErrorCount == 0;
 }
 
 #if DO_CHECK
@@ -5787,68 +6230,115 @@ public:
 };
 #endif //DO_CHECK
 
-
-void FPakFile::GetFilenames(TArray<FString>& OutFileList) const
+void FPakFile::GetPrunedFilenames(TArray<FString>& OutFileList) const
 {
-	for (const TMemoryImageMap<FMemoryImageString, FPakDirectory>::ElementType& DirectoryElement : Data->Index)
+	for (FFilenameIterator It(*this, true /* bIncludeDeleted */); It; ++It)
 	{
-		const  FPakDirectory& Directory = DirectoryElement.Value;
-		for (const FPakDirectory::ElementType& FileElement : Directory)
-		{
-			OutFileList.Add(MountPoint / DirectoryElement.Key / FileElement.Key);
-		}
+		OutFileList.Add(PakPathCombine(MountPoint, It.Filename()));
 	}
 }
 
-
-void FPakFile::GetFilenamesInChunk(const TArray<int32>& InChunkIDs, TArray<FString>& OutFileList)
+void FPakFile::GetPrunedFilenamesInChunk(const TArray<int32>& InChunkIDs, TArray<FString>& OutFileList) const
 {
-	TSet<int32> OverlappingEntries;
-
-	for (int32 LocalChunkID : InChunkIDs)
+	for (FFilenameIterator It(*this, true /* bIncludeDeleted */); It; ++It)
 	{
-		int32 ChunkStart = LocalChunkID * FPakInfo::MaxChunkDataSize;
-		int32 ChunkEnd = ChunkStart + FPakInfo::MaxChunkDataSize;
-		int32 FileIndex = 0;
+		const FPakEntry& File = It.Info();
+		int64 FileStart = File.Offset;
+		int64 FileEnd = File.Offset + File.Size;
 
-		for (const FPakEntry& File : Data->Files)
+		for (int64 LocalChunkID : InChunkIDs)
 		{
-			int32 FileStart = File.Offset;
-			int32 FileEnd = File.Offset + File.Size;
+			int64 ChunkStart = LocalChunkID * FPakInfo::MaxChunkDataSize;
+			int64 ChunkEnd = ChunkStart + FPakInfo::MaxChunkDataSize;
 
-			// If this file is past the end of the target chunk, we're done
-			if (FileStart > ChunkEnd)
+			if (FileStart < ChunkEnd && FileEnd > ChunkStart)
 			{
+				OutFileList.Add(It.Filename());
 				break;
 			}
+		}
+	}
+}
+
+#if ENABLE_PAKFILE_RUNTIME_PRUNING_VALIDATE
+void FPakFile::ValidateDirectorySearch(const TSet<FString>& FullFoundFiles, const TSet<FString>& PrunedFoundFiles, const TCHAR* InPath) const
+{
+	TArray<FString> MissingFromPruned;
+	for (const FString& FileInFull : FullFoundFiles)
+	{
+		if (!PrunedFoundFiles.Contains(FileInFull))
+		{
+			MissingFromPruned.Add(FileInFull);
+		}
+	}
+	TArray<FString> MissingFromFull;
+	for (const FString& FileInPruned : PrunedFoundFiles)
+	{
+		if (!FullFoundFiles.Contains(FileInPruned))
+		{
+			MissingFromFull.Add(FileInPruned);
+		}
+	}
 
+	if (MissingFromPruned.Num() == 0 && MissingFromFull.Num() == 0)
+	{
+		return;
+	}
 
-			if (FileEnd > ChunkStart)
+	TArray<FString> WildCards, OldWildCards;
+	GConfig->GetArray(TEXT("Pak"), TEXT("IndexValidationIgnore"), WildCards, GEngineIni);
+	auto IsIgnore = [&WildCards](const FString& FilePath)
+	{
+		for (const FString& WildCard : WildCards)
+		{
+			if (FilePath.MatchesWildcard(WildCard))
 			{
-				OverlappingEntries.Add(FileIndex);
+				return true;
 			}
-
-			FileIndex++;
 		}
-	}
-
-	int32 Remaining = OverlappingEntries.Num();
-	for (const TMemoryImageMap<FMemoryImageString, FPakDirectory>::ElementType& DirectoryElement : Data->Index)
+		return false;
+	};
+	auto StripIgnores = [&IsIgnore](TArray<FString>& FilePaths)
 	{
-		const  FPakDirectory& Directory = DirectoryElement.Value;
-		for (const FPakDirectory::ElementType& FileElement : Directory)
+		for (int Idx = FilePaths.Num() - 1; Idx >= 0; --Idx)
 		{
-			if (OverlappingEntries.Contains(FileElement.Value))
+			if (IsIgnore(FilePaths[Idx]))
 			{
-				OutFileList.Add(FString(DirectoryElement.Key) / FString(FileElement.Key));
-				if (--Remaining == 0)
-				{
-					break;
-				}
+				FilePaths.RemoveAtSwap(Idx);
 			}
 		}
+	};
+
+	StripIgnores(MissingFromPruned);
+	StripIgnores(MissingFromFull);
+
+	if (MissingFromPruned.Num() == 0 && MissingFromFull.Num() == 0)
+	{
+		return;
+	}
+	MissingFromPruned.Sort();
+	MissingFromFull.Sort();
+
+	UE_LOG(LogPakFile, Error, TEXT("FindPrunedFilesAtPath('%s') for PakFile '%s' found a different list in the FullDirectory than in the PrunedDirectory. ")
+		TEXT("Change the calling code or add the files to Engine:[Pak]:WildcardsToKeepInPakStringIndex or Engine:[Pak]:IndexValidationIgnore."),
+		InPath, *PakFilename);
+	if (MissingFromPruned.Num() > 0)
+	{
+		for (const FString& Missing : MissingFromPruned)
+		{
+			UE_LOG(LogPakFile, Warning, TEXT("MissingPrunedPakFile: %s"), *Missing);
+		}
+	}
+	if (MissingFromFull.Num() > 0)
+	{
+		UE_LOG(LogPakFile, Error, TEXT("Some files in the PrunedDirectory are missing from the FullDirectory.  This is a logic error in FPakFile since the PrunedDirectory should be a subset of the FullDirectory."));
+		for (const FString& Missing : MissingFromFull)
+		{
+			UE_LOG(LogPakFile, Warning, TEXT("MissingFullPakFile: %s"), *Missing);
+		}
 	}
 }
+#endif
 
 FArchive* FPakFile::GetSharedReader(IPlatformFile* LowerLevel)
 {
@@ -5895,150 +6385,82 @@ FArchive* FPakFile::GetSharedReader(IPlatformFile* LowerLevel)
 	return PakReader;
 }
 
-FPakFile::EFindResult FPakFile::Find(const FString& Filename, FPakEntry* OutEntry) const
+const FPakEntryLocation* FPakFile::FindLocationFromIndex(const FString& FullPath, const FString& MountPoint, const FPathHashIndex& PathHashIndex, uint64 PathHashSeed)
 {
-	QUICK_SCOPE_CYCLE_COUNTER(PakFileFind);
-	if (Filename.StartsWith(MountPoint))
+	if (!FullPath.StartsWith(MountPoint))
 	{
-		FString Path(FPaths::GetPath(Filename));
-
-		// Handle the case where the user called FPakFile::UnloadFilenames() and the filenames
-		// were removed from memory.
-		if (bFilenamesRemoved)
-		{
-			// Derived from the following:
-			//     FString RelativeFilename(Filename.Mid(Path.Len() + 1));
-			//     Path = Path.Mid(MountPoint.Len()) / RelativeFilename;
-			// Hash the Path.
-			int AdjustedMountPointLen = Path.Len() < MountPoint.Len() ? Path.Len() : MountPoint.Len();
-			FString LowercaseFilename = Filename.ToLower();
-			const TCHAR* SplitStartPtr = *LowercaseFilename + AdjustedMountPointLen;
-			uint32 SplitLen = LowercaseFilename.Len() - AdjustedMountPointLen;
-			if (*SplitStartPtr == '/')
-			{
-				++SplitStartPtr;
-				--SplitLen;
-			}
-			uint64 PathHash = FFnv::MemFnv64(SplitStartPtr, SplitLen * sizeof(TCHAR), FilenameStartHash);
-
-			// Look it up in our sorted-by-filename-hash array.
-			uint32 PathHashMostSignificantBits = PathHash >> 56;
-			uint32 HashEntriesCount = FilenameHashesIndex[PathHashMostSignificantBits + 1] - FilenameHashesIndex[PathHashMostSignificantBits];
-			uint64* FoundHash = (uint64*)bsearch(&PathHash, FilenameHashes + FilenameHashesIndex[PathHashMostSignificantBits], HashEntriesCount, sizeof(uint64), CompareFilenameHashes);
-			if (FoundHash != NULL)
-			{
-				bool bDeleted = false;
-
-				int32 FoundEntryIndex = FilenameHashesIndices[FoundHash - FilenameHashes];
-
-				if (MiniPakEntries != NULL)
-				{
-					uint32 MemoryOffset = MiniPakEntriesOffsets[FoundEntryIndex];
+		return nullptr;
+	}
+	const TCHAR* RelativePathFromMount = (*FullPath) + MountPoint.Len();
+	uint64 PathHash = HashPath(RelativePathFromMount, PathHashSeed);
+	return PathHashIndex.Find(PathHash);
+}
 
-					bDeleted = (MemoryOffset == MAX_uint32); // deleted records have a magic number in the offset instead (not ideal, but there is no more space in the bit-encoded entry)
+const FPakEntryLocation* FPakFile::FindLocationFromIndex(const FString& FullPath, const FString& MountPoint, const FDirectoryIndex& DirectoryIndex)
+{
+	if (!FullPath.StartsWith(MountPoint))
+	{
+		return nullptr;
+	}
+	const TCHAR* RelativePathFromMount = (*FullPath) + MountPoint.Len();
+	FString RelativeDirName(RelativePathFromMount);
+	FString CleanFileName;
+	SplitPathInline(RelativeDirName, CleanFileName);
+	const FPakDirectory* PakDirectory = DirectoryIndex.Find(RelativeDirName);
+	if (PakDirectory)
+	{
+		return PakDirectory->Find(CleanFileName);
+	}
+	return nullptr;
+}
 
-					if (OutEntry != NULL)
-					{
-						if (!bDeleted)
-						{
-							// The FPakEntry structures are bit-encoded, so decode it.
-							DecodePakEntry(MiniPakEntries + MemoryOffset, OutEntry);
-						}
-						else
-						{
-							// entry was deleted and original data is inaccessible- build dummy entry
-							(*OutEntry) = FPakEntry();
-							OutEntry->SetDeleteRecord(true);
-							OutEntry->Verified = true;		// Set Verified to true to avoid have a synchronous open fail comparing FPakEntry structures.
-						}
-					}
-				}
-				else
-				{
-					const FPakEntry* FoundEntry = &Data->Files[FoundEntryIndex];
+FPakFile::EFindResult FPakFile::Find(const FString& FullPath, FPakEntry* OutEntry) const
+{
+	QUICK_SCOPE_CYCLE_COUNTER(PakFileFind);
 
-					bDeleted = FoundEntry->IsDeleteRecord();
+	const FPakEntryLocation* PakEntryLocation;
+#if ENABLE_PAKFILE_RUNTIME_PRUNING_VALIDATE
+	if (IsPakValidatePruning() && bHasPathHashIndex && bHasFullDirectoryIndex)
+	{
+		const FPakEntryLocation* PathHashLocation = nullptr;
+		PathHashLocation = FindLocationFromIndex(FullPath, MountPoint, PathHashIndex, PathHashSeed);
 
-					if (OutEntry != NULL)
-					{
-						OutEntry->Offset = FoundEntry->Offset;
-						OutEntry->Size = FoundEntry->Size;
-						OutEntry->UncompressedSize = FoundEntry->UncompressedSize;
-						OutEntry->CompressionMethodIndex = FoundEntry->CompressionMethodIndex;
-						// NEEDED? FMemory::Memcpy(OutEntry->Hash, FoundEntry->Hash, sizeof(OutEntry->Hash));
-						OutEntry->CompressionBlocks = FoundEntry->CompressionBlocks;
-						OutEntry->CompressionBlockSize = FoundEntry->CompressionBlockSize;
-						OutEntry->Flags = FoundEntry->Flags;
-						OutEntry->Verified = true;		// Set Verified to true to avoid have a synchronous open fail comparing FPakEntry structures.
-					}
-				}
+		const FPakEntryLocation* DirectoryLocation = nullptr;
+		DirectoryLocation = FindLocationFromIndex(FullPath, MountPoint, DirectoryIndex);
 
-				return bDeleted ? EFindResult::FoundDeleted : EFindResult::Found;
+		if ((PathHashLocation != nullptr) != (DirectoryLocation != nullptr))
+		{
+			const TCHAR* FoundName = TEXT("PathHashIndex");
+			const TCHAR* NotFoundName = TEXT("FullDirectoryIndex");
+			if (!PathHashLocation)
+			{
+				Swap(FoundName, NotFoundName);
 			}
+			UE_LOG(LogPakFile, Error, TEXT("PathHashIndex does not match FullDirectoryIndex. Pakfile '%s' has '%s' in its %s but not in its %s."),
+				*PakFilename, *FullPath, FoundName, NotFoundName);
+		}
+		PakEntryLocation = PathHashLocation ? PathHashLocation : DirectoryLocation;
+	}
+	else
+#endif
+	{
+		if (bHasPathHashIndex)
+		{
+			PakEntryLocation = FindLocationFromIndex(FullPath, MountPoint, PathHashIndex, PathHashSeed);
 		}
 		else
 		{
-			const FPakDirectory* PakDirectory = FindDirectory(*Path);
-			if (PakDirectory != NULL)
-			{
-				FMemoryImageString RelativeFilename(Filename.Mid(Path.Len() + 1));
-				int32 const* FoundEntryIndex = PakDirectory->Find(RelativeFilename);
-				if (FoundEntryIndex != NULL)
-				{
-					bool bDeleted = false;
-
-					if (MiniPakEntries != NULL)
-					{
-						// The FPakEntry structures are bit-encoded, so decode it.
-						uint32 MemoryOffset = MiniPakEntriesOffsets[*FoundEntryIndex];
-
-						bDeleted = (MemoryOffset == MAX_uint32); // deleted records have a magic number in the offset instead (not ideal, but there is no more space in the bit-encoded entry)
-
-						if (OutEntry != NULL)
-						{
-							if (!bDeleted)
-							{
-								// The FPakEntry structures are bit-encoded, so decode it.
-								uint8* FoundPtr = MiniPakEntries + MemoryOffset;
-								DecodePakEntry(FoundPtr, OutEntry);
-							}
-							else
-							{
-								// entry was deleted and original data is inaccessible- build dummy entry
-								(*OutEntry) = FPakEntry();
-								OutEntry->SetDeleteRecord(true);
-								OutEntry->Verified = true;		// Set Verified to true to avoid have a synchronous open fail comparing FPakEntry structures.
-							}
-						}
-					}
-					else
-					{
-						const FPakEntry* FoundEntry = &Data->Files[*FoundEntryIndex];
-						bDeleted = FoundEntry->IsDeleteRecord();
-
-						if (OutEntry != NULL)
-						{
-							//*OutEntry = **FoundEntry;
-							OutEntry->Offset = FoundEntry->Offset;
-							OutEntry->Size = FoundEntry->Size;
-							OutEntry->UncompressedSize = FoundEntry->UncompressedSize;
-							OutEntry->CompressionMethodIndex = FoundEntry->CompressionMethodIndex;
-							FMemory::Memcpy(OutEntry->Hash, FoundEntry->Hash, sizeof(OutEntry->Hash));
-							OutEntry->CompressionBlocks = FoundEntry->CompressionBlocks;
-							OutEntry->CompressionBlockSize = FoundEntry->CompressionBlockSize;
-							OutEntry->Flags = FoundEntry->Flags;
-							OutEntry->Verified = true;		// Set Verified to true to avoid have a synchronous open fail comparing FPakEntry structures.
-						}
-					}
-
-					return bDeleted ? EFindResult::FoundDeleted : EFindResult::Found;
-				}
-			}
+			check(bHasFullDirectoryIndex);
+			PakEntryLocation = FindLocationFromIndex(FullPath, MountPoint, DirectoryIndex);
 		}
 	}
-	return EFindResult::NotFound;
-}
+	if (!PakEntryLocation)
+	{
+		return EFindResult::NotFound;
+	}
 
+	return GetPakEntry(*PakEntryLocation, OutEntry);
+}
 
 #if !UE_BUILD_SHIPPING
 class FPakExec : private FSelfRegisteringExec
@@ -6130,7 +6552,10 @@ FPakPlatformFile::~FPakPlatformFile()
 	FCoreDelegates::OnFEngineLoopInitComplete.RemoveAll(this);
 
 	FCoreDelegates::OnMountAllPakFiles.Unbind();
+	PRAGMA_DISABLE_DEPRECATION_WARNINGS
 	FCoreDelegates::OnMountPak.Unbind();
+	PRAGMA_ENABLE_DEPRECATION_WARNINGS
+	FCoreDelegates::MountPak.Unbind();
 	FCoreDelegates::OnUnmountPak.Unbind();
 	FCoreDelegates::OnOptimizeMemoryUsageForMountedPaks.Unbind();
 
@@ -6298,7 +6723,10 @@ bool FPakPlatformFile::Initialize(IPlatformFile* Inner, const TCHAR* CmdLine)
 #endif // !UE_BUILD_SHIPPING
 
 	FCoreDelegates::OnMountAllPakFiles.BindRaw(this, &FPakPlatformFile::MountAllPakFiles);
-	FCoreDelegates::OnMountPak.BindRaw(this, &FPakPlatformFile::HandleMountPakDelegate);
+	PRAGMA_DISABLE_DEPRECATION_WARNINGS
+	FCoreDelegates::OnMountPak.BindRaw(this, &FPakPlatformFile::HandleOnMountPakDelegate);
+	PRAGMA_ENABLE_DEPRECATION_WARNINGS
+	FCoreDelegates::MountPak.BindRaw(this, &FPakPlatformFile::HandleMountPakDelegate);
 	FCoreDelegates::OnUnmountPak.BindRaw(this, &FPakPlatformFile::HandleUnmountPakDelegate);
 	FCoreDelegates::OnOptimizeMemoryUsageForMountedPaks.BindRaw(this, &FPakPlatformFile::OptimizeMemoryUsageForMountedPaks);
 
@@ -6324,37 +6752,94 @@ void FPakPlatformFile::InitializeNewAsyncIO()
 #endif
 }
 
+#if ENABLE_PAKFILE_RUNTIME_PRUNING
+bool FPakFile::bSomePakNeedsPruning = false;
+#endif
+
+#if !UE_BUILD_SHIPPING
+int32 GetRecursiveAllocatedSize(const FPakDirectory& Index)
+{
+	int32 Size = Index.GetAllocatedSize();
+	for (const TPair<FString, FPakEntryLocation>& kvpair : Index)
+	{
+		Size += kvpair.Key.GetAllocatedSize();
+	}
+	return Size;
+}
+int32 GetRecursiveAllocatedSize(const FPakFile::FDirectoryIndex& Index)
+{
+	int32 Size = Index.GetAllocatedSize();
+	for (const TPair<FString, FPakDirectory>& kvpair : Index)
+	{
+		Size += kvpair.Key.GetAllocatedSize();
+		Size += GetRecursiveAllocatedSize(kvpair.Value);
+	}
+	return Size;
+}
+#endif
+
 void FPakPlatformFile::OptimizeMemoryUsageForMountedPaks()
 {
-#if !(IS_PROGRAM || WITH_EDITOR)
-	FSlowHeartBeatScope SuspendHeartBeat;
-	bool bUnloadPakEntryFilenamesIfPossible = FParse::Param(FCommandLine::Get(), TEXT("unloadpakentryfilenames"));
-	GConfig->GetBool(TEXT("Pak"), TEXT("UnloadPakEntryFilenamesIfPossible"), bUnloadPakEntryFilenamesIfPossible, GEngineIni);
+#if ENABLE_PAKFILE_RUNTIME_PRUNING || !UE_BUILD_SHIPPING
+	TArray<FPakListEntry> Paks;
+	bool bNeedsPaks = false;
+#endif
+
+#if !UE_BUILD_SHIPPING
+	bNeedsPaks = true;
+#endif
+#if ENABLE_PAKFILE_RUNTIME_PRUNING 
+	bNeedsPaks = bNeedsPaks || FPakFile::bSomePakNeedsPruning;
+#endif
+	if (bNeedsPaks)
+	{
+		GetMountedPaks(Paks);
+	}
 
-	if ((bUnloadPakEntryFilenamesIfPossible && !FParse::Param(FCommandLine::Get(), TEXT("nounloadpakentries"))) || FParse::Param(FCommandLine::Get(), TEXT("unloadpakentries")))
+#if ENABLE_PAKFILE_RUNTIME_PRUNING
+	if (FPakFile::bSomePakNeedsPruning)
 	{
-		// With [Pak] UnloadPakEntryFilenamesIfPossible enabled, [Pak] DirectoryRootsToKeepInMemoryWhenUnloadingPakEntryFilenames
-		// can contain pak entry directory wildcards of which the entire recursive directory structure of filenames underneath a
-		// matching wildcard will be kept.
-		//
-		// Example:
-		//   [Pak]
-		//   DirectoryRootsToKeepInMemoryWhenUnloadingPakEntryFilenames="*/Config/Tags/"
-		//   +DirectoryRootsToKeepInMemoryWhenUnloadingPakEntryFilenames="*/Content/Localization/*"
-		TArray<FString> DirectoryRootsToKeep;
-		GConfig->GetArray(TEXT("Pak"), TEXT("DirectoryRootsToKeepInMemoryWhenUnloadingPakEntryFilenames"), DirectoryRootsToKeep, GEngineIni);
+		for (auto& Pak : Paks)
+		{
+			FPakFile* PakFile = Pak.PakFile;
+			if (PakFile->bWillPruneDirectoryIndex)
+			{
+				check(PakFile->bHasPathHashIndex);
+				FWriteScopeLock DirectoryLock(PakFile->DirectoryIndexLock);
+				if (PakFile->bNeedsLegacyPruning)
+				{
+					FPakFile::PruneDirectoryIndex(PakFile->DirectoryIndex, &PakFile->PrunedDirectoryIndex, PakFile->MountPoint);
+					PakFile->bNeedsLegacyPruning = false;
+				}
 
-		FPakPlatformFile* PakPlatformFile = (FPakPlatformFile*)(FPlatformFileManager::Get().FindPlatformFile(FPakPlatformFile::GetTypeName()));
-		PakPlatformFile->UnloadPakEntryFilenames(&DirectoryRootsToKeep);
+				Swap(PakFile->DirectoryIndex, PakFile->PrunedDirectoryIndex);
+				PakFile->PrunedDirectoryIndex.Empty();
+				PakFile->bHasFullDirectoryIndex = false;
+				PakFile->bWillPruneDirectoryIndex = false;
+			}
+		}
 	}
+#endif
 
-	bool bShrinkPakEntriesMemoryUsage = FParse::Param(FCommandLine::Get(), TEXT("shrinkpakentries"));
-	GConfig->GetBool(TEXT("Pak"), TEXT("ShrinkPakEntriesMemoryUsage"), bShrinkPakEntriesMemoryUsage, GEngineIni);
-	if (bShrinkPakEntriesMemoryUsage)
+#if !UE_BUILD_SHIPPING
+	uint32 DirectoryHashSize = 0;
+	uint32 PathHashSize = 0;
+	uint32 EntriesSize = 0;
+
+	for (auto& Pak : Paks)
 	{
-		FPakPlatformFile* PakPlatformFile = (FPakPlatformFile*)(FPlatformFileManager::Get().FindPlatformFile(FPakPlatformFile::GetTypeName()));
-		PakPlatformFile->ShrinkPakEntriesMemoryUsage();
+		FPakFile* PakFile = Pak.PakFile;
+		DirectoryHashSize += GetRecursiveAllocatedSize(PakFile->DirectoryIndex);
+#if ENABLE_PAKFILE_RUNTIME_PRUNING 
+		{
+			DirectoryHashSize += GetRecursiveAllocatedSize(PakFile->PrunedDirectoryIndex);
+		}
+#endif
+		PathHashSize += PakFile->PathHashIndex.GetAllocatedSize();
+		EntriesSize += PakFile->EncodedPakEntries.GetAllocatedSize();
+		EntriesSize += PakFile->Files.GetAllocatedSize();
 	}
+	UE_LOG(LogPakFile, Log, TEXT("AllPaks IndexSizes: DirectoryHashSize=%d, PathHashSize=%d, EntriesSize=%d, TotalSize=%d"), DirectoryHashSize, PathHashSize, EntriesSize, DirectoryHashSize + PathHashSize + EntriesSize);
 #endif
 }
 
@@ -6454,17 +6939,14 @@ bool FPakPlatformFile::Mount(const TCHAR* InPakFilename, uint32 PakOrder, const
 
 			PRAGMA_DISABLE_DEPRECATION_WARNINGS
 			FCoreDelegates::PakFileMountedCallback.Broadcast(InPakFilename);
-			PRAGMA_ENABLE_DEPRECATION_WARNINGS
 			FCoreDelegates::OnPakFileMounted.Broadcast(InPakFilename, Pak->PakchunkIndex);
-			if (FCoreDelegates::NewFileAddedDelegate.IsBound())
+			PRAGMA_ENABLE_DEPRECATION_WARNINGS
+			static double OnPakFileMounted2Time = 0.0;
 			{
-				TArray<FString> Filenames;
-				Pak->GetFilenames(Filenames);
-				for (const FString& Filename : Filenames)
-				{
-					FCoreDelegates::NewFileAddedDelegate.Broadcast(Filename);
-				}
+				FScopedDurationTimer Timer(OnPakFileMounted2Time);
+				FCoreDelegates::OnPakFileMounted2.Broadcast(*Pak);
 			}
+			UE_LOG(LogPakFile, Log, TEXT("OnPakFileMounted2Time == %lf"), OnPakFileMounted2Time);
 		}
 		else
 		{
@@ -6643,7 +7125,7 @@ int32 FPakPlatformFile::GetPakOrderFromPakFilePath(const FString& PakFilePath)
 	return 0;
 }
 
-bool FPakPlatformFile::HandleMountPakDelegate(const FString& PakFilePath, int32 PakOrder, IPlatformFile::FDirectoryVisitor* Visitor)
+IPakFile* FPakPlatformFile::HandleMountPakDelegate(const FString& PakFilePath, int32 PakOrder)
 {
 	FPlatformMisc::LowLevelOutputDebugStringf(TEXT("Mounting pak file: %s \n"), *PakFilePath);
 
@@ -6652,8 +7134,7 @@ bool FPakPlatformFile::HandleMountPakDelegate(const FString& PakFilePath, int32
 		PakOrder = GetPakOrderFromPakFilePath(PakFilePath);
 	}
 	
-	bool bReturn = Mount(*PakFilePath, PakOrder);
-	if (bReturn && Visitor != nullptr)
+	if (Mount(*PakFilePath, PakOrder))
 	{
 		TArray<FPakListEntry> Paks;
 		GetMountedPaks(Paks);
@@ -6662,16 +7143,30 @@ bool FPakPlatformFile::HandleMountPakDelegate(const FString& PakFilePath, int32
 		{
 			if (PakFilePath == Pak.PakFile->GetFilename())
 			{
-				// Get a list of all of the files in the pak
-				for (FPakFile::FFileIterator It(*Pak.PakFile); It; ++It)
-				{
-					Visitor->Visit(*It.Filename(), false);
-				}
-				return true;
+				return Pak.PakFile;
 			}
 		}
+
+		check(false); // Should have found the pak since Mount returned true
+	}
+	return nullptr;
+}
+
+bool FPakPlatformFile::HandleOnMountPakDelegate(const FString& PakFilePath, int32 PakOrder, IPlatformFile::FDirectoryVisitor* Visitor)
+{
+	IPakFile* PakFile = HandleMountPakDelegate(PakFilePath, PakOrder);
+	if (PakFile)
+	{
+		if (Visitor)
+		{
+			PakFile->PakVisitPrunedFilenames(*Visitor);
+		}
+		return true;
+	}
+	else
+	{
+		return false;
 	}
-	return bReturn;
 }
 
 bool FPakPlatformFile::HandleUnmountPakDelegate(const FString& PakFilePath)
@@ -6865,74 +7360,6 @@ bool FPakPlatformFile::CopyFile(const TCHAR* To, const TCHAR* From, EPlatformFil
 	return Result;
 }
 
-void FPakPlatformFile::UnloadPakEntryFilenames(TArray<FString>* DirectoryRootsToKeep)
-{
-	int32 TotalNumFilenames = 0;
-	int32 NumFilenamesUnloaded = 0;
-	int32 NumPaks = 0;
-	double Timer = 0.0;
-	{
-		SCOPE_SECONDS_COUNTER(Timer);
-
-		TArray<FPakListEntry> Paks;
-		GetMountedPaks(Paks);
-		TMap<uint64, FPakEntry> CrossPakCollisionDetector;
-		
-		for (auto Pak : Paks)
-		{
-			if (Pak.PakFile->HasFilenames())
-			{
-				TotalNumFilenames += Pak.PakFile->GetNumFiles();
-			}
-		}
-
-		CrossPakCollisionDetector.Reserve(TotalNumFilenames);
-
-		// Sort the pak list by number of entries so that we allow the larger ones a better chance of not encountering a collision
-		Algo::Sort(Paks, [](const FPakListEntry& A, const FPakListEntry& B) { return A.PakFile->GetNumFiles() > B.PakFile->GetNumFiles(); });
-
-		for (auto Pak : Paks)
-		{
-			if (Pak.PakFile->HasFilenames())
-			{
-				NumPaks++;
-				int32 NumFilesInThisPak = Pak.PakFile->GetNumFiles();
-
-				if (Pak.PakFile->UnloadPakEntryFilenames(CrossPakCollisionDetector, DirectoryRootsToKeep))
-				{
-					NumFilenamesUnloaded += NumFilesInThisPak;
-				}
-			}
-		}
-	}
-	UE_LOG(LogPakFile, Log, TEXT("Unloaded %d/%d filenames from %d pak files in %.4fs"), NumFilenamesUnloaded, TotalNumFilenames, NumPaks, Timer);
-}
-
-void FPakPlatformFile::ShrinkPakEntriesMemoryUsage()
-{
-	double Timer = 0.0;
-	int32 NumPakFiles = 0;
-	int32 NumEntries = 0;
-	{
-		SCOPE_SECONDS_COUNTER(Timer);
-
-		TArray<FPakListEntry> Paks;
-		GetMountedPaks(Paks);
-		for (auto Pak : Paks)
-		{
-			if (!Pak.PakFile->HasShrunkPakEntries())
-			{
-				NumPakFiles++;
-				if (Pak.PakFile->ShrinkPakEntriesMemoryUsage())
-				{
-					NumEntries += Pak.PakFile->GetNumFiles();
-				}
-			}
-		}
-	}
-	UE_LOG(LogPakFile, Log, TEXT("Shrunk %d entries from %d pak files in %.4fs"), NumEntries, NumPakFiles, Timer);
-}
-
 /**
  * Module for the pak file
  */
@@ -6968,62 +7395,39 @@ public:
 	TUniquePtr<IPlatformFile> Singleton;
 };
 
-void FPakFile::AddSpecialFile(FPakEntry Entry, const FString& Filename)
+void FPakFile::AddSpecialFile(const FPakEntry& Entry, const FString& Filename)
 {
 	MakeDirectoryFromPath(MountPoint);
 
-	int32 EntryIndex = Data->Files.Num();
-
-	// Add new file info.
-	Data->Files.Add(Entry);
-	NumEntries++;
-
-	// Construct Index of all directories in pak file.
-
-
-	FString CleanFilename(Filename.Mid(MountPoint.Len() ));
-
-	FString Path = FPaths::GetPath(CleanFilename);
-	FString Path2 = FPaths::GetPath(Filename);
-	MakeDirectoryFromPath(Path);
-
-
-
-	FPakDirectory* Directory = Data->Index.Find(Path);
-	if (Directory != NULL)
-	{
-		Directory->Add(FPaths::GetCleanFilename(CleanFilename), EntryIndex);
-	}
-	else
+	// TODO: This function is not threadsafe; readers of the Indexes will be invalidated when we modify them
+	// To make it threadsafe would require always holding the lock around any read of either index, which is
+	// more expensive than we want to support this debug feature
+	FPakEntryLocation EntryLocation;
+	if (!Entry.IsDeleteRecord())
 	{
-		FPakDirectory& NewDirectory = Data->Index.Add(Path);
-		NewDirectory.Add(FPaths::GetCleanFilename(CleanFilename), EntryIndex);
-
-		// add the parent directories up to the mount point
-		while (MountPoint != Path)
+		// Add new file info.
+		TArray<uint8> NewEncodedPakEntries;
+		FMemoryWriter MemoryWriter(NewEncodedPakEntries);
+		EntryLocation = FPakEntryLocation::CreateFromOffsetIntoEncoded(EncodedPakEntries.Num());
+		if (EncodePakEntry(MemoryWriter, Entry, Info))
 		{
-			Path.LeftInline(Path.Len() - 1, false);
-			int32 Offset = 0;
-			if (Path.FindLastChar('/', Offset))
-			{
-				Path.LeftInline(Offset, false);
-				MakeDirectoryFromPath(Path);
-				if (Data->Index.Find(Path) == NULL)
-				{
-					Data->Index.Add(Path);
-				}
-			}
-			else
-			{
-				Path = MountPoint;
-			}
+			EncodedPakEntries.Append(NewEncodedPakEntries);
+			EncodedPakEntries.Shrink();
+		}
+		else
+		{
+			EntryLocation = FPakEntryLocation::CreateFromListIndex(Files.Num());
+			Files.Add(Entry);
+			Files.Shrink();
 		}
+		NumEntries++;
 	}
-}
-
 
+	FPathHashIndex* PathHashToWrite = bHasPathHashIndex ? &PathHashIndex : nullptr;
+	AddEntryToIndex(Filename, EntryLocation, MountPoint, PathHashSeed, &DirectoryIndex, PathHashToWrite, nullptr /* CollisionDetection */);
+}
 
-void FPakPlatformFile::MakeUniquePakFilesForTheseFiles(TArray<TArray<FString>> InFiles)
+void FPakPlatformFile::MakeUniquePakFilesForTheseFiles(const TArray<TArray<FString>>& InFiles)
 {
 	for (int k = 0; k < InFiles.Num(); k++)
 	{
@@ -7031,24 +7435,33 @@ void FPakPlatformFile::MakeUniquePakFilesForTheseFiles(TArray<TArray<FString>> I
 		for (int i = 0; i < InFiles[k].Num(); i++)
 		{
 			FPakEntry FileEntry;
-			FPakFile* PakFile = NULL;
-			bool bFoundEntry = FindFileInPakFiles(*InFiles[k][i], &PakFile, &FileEntry);
-			if (bFoundEntry && PakFile && PakFile->GetFilenameName() != NAME_None)
+			FPakFile* ExistingRealPakFile = nullptr;
+			bool bFoundEntry = FindFileInPakFiles(*InFiles[k][i], &ExistingRealPakFile, &FileEntry);
+			if (bFoundEntry && ExistingRealPakFile && ExistingRealPakFile->GetFilenameName() != NAME_None)
 			{
-				if (NewPakFile == nullptr && Mount(*PakFile->GetFilename(), 500, *PakFile->MountPoint, false))
+				if (NewPakFile == nullptr)
 				{
-					// we successfully mounted the file, find the empty pak file we just added.
-					for (int j = 0; j < PakFiles.Num(); j++)
+					// Mount another copy of the existing real PakFile, but don't allow it to Load the index, so it intializes as empty
+					const bool bLoadIndex = false;
+					if (Mount(*ExistingRealPakFile->GetFilename(), 500, *ExistingRealPakFile->MountPoint, bLoadIndex))
 					{
-						if (PakFiles[j].PakFile->Data->Files.Num() == 0 && PakFiles[j].PakFile->CachedTotalSize == PakFile->CachedTotalSize)
+						// we successfully mounted the file, find the empty pak file we just added.
+						for (int j = 0; j < PakFiles.Num(); j++)
 						{
-							NewPakFile = PakFiles[j].PakFile;
-							break;
+							FPakFile& PotentialNewPakFile = *PakFiles[j].PakFile;
+							if (PotentialNewPakFile.PakFilename == ExistingRealPakFile->PakFilename &&  // It has the right name
+								PotentialNewPakFile.CachedTotalSize == ExistingRealPakFile->CachedTotalSize && // And it has the right size
+								PotentialNewPakFile.GetNumFiles() == 0) // And it didn't load the index
+							{
+								NewPakFile = &PotentialNewPakFile;
+								break;
+							}
+						}
+
+						if (NewPakFile != nullptr)
+						{
+							NewPakFile->SetCacheType(FPakFile::ECacheType::Individual);
 						}
-					}
-					if (NewPakFile != nullptr)
-					{
-						NewPakFile->SetCacheType(FPakFile::ECacheType::Individual);
 					}
 				}
 
