Optimize Pak file indices by computing the PathHashIndex in unrealpak rather than using strings at runtime.

diff --git a/Engine/Source/Runtime/PakFile/Public/IPlatformFilePak.h b/Engine/Source/Runtime/PakFile/Public/IPlatformFilePak.h
--- a/Engine/Source/Runtime/PakFile/Public/IPlatformFilePak.h
+++ b/Engine/Source/Runtime/PakFile/Public/IPlatformFilePak.h
@@ -25,6 +25,17 @@ DECLARE_DWORD_ACCUMULATOR_STAT_EXTERN(TEXT("Num open pak file handles"), STAT_Pa
 
 #define PAK_TRACKER 0
 
+// ENABLE_PAKFILE_RUNTIME_PRUNING allows pruning the DirectoryIndex at runtime after all Paks have loaded rather than loading only the already-pruned DirectoryIndex
+// This requires extra cputime to make reads of the DirectoryIndex ThreadSafe, and will be removed in a future version
+#ifndef ENABLE_PAKFILE_RUNTIME_PRUNING
+#define ENABLE_PAKFILE_RUNTIME_PRUNING 1
+#endif
+#define ENABLE_PAKFILE_RUNTIME_PRUNING_VALIDATE ENABLE_PAKFILE_RUNTIME_PRUNING && !UE_BUILD_SHIPPING
+
+#if ENABLE_PAKFILE_RUNTIME_PRUNING
+#include "Misc/ScopeRWLock.h"
+#endif
+
 // Define the type of a chunk hash. Currently selectable between SHA1 and CRC32.
 #define PAKHASH_USE_CRC	1
 #if PAKHASH_USE_CRC
@@ -64,6 +75,8 @@ DECLARE_DELEGATE_RetVal_OneParam(bool, FFilenameSecurityDelegate, const TCHAR* /
 DECLARE_DELEGATE_ThreeParams(FPakCustomEncryptionDelegate, uint8* /*InData*/, uint32 /*InDataSize*/, FGuid /*InEncryptionKeyGuid*/);
 DECLARE_MULTICAST_DELEGATE_OneParam(FPakChunkSignatureCheckFailedHandler, const FPakChunkSignatureCheckFailedData&);
 DECLARE_MULTICAST_DELEGATE_OneParam(FPakMasterSignatureTableCheckFailureHandler, const FString&);
+/** Delegate which allows a project to configure Index Pruning.  This is a delegate instead of a config file because config files are loaded after the first PakFiles */
+DECLARE_DELEGATE_ThreeParams(FPakSetIndexSettings, bool& /* bKeepFullDirectory*/, bool& /* bValidatePruning */, bool& /* bDelayPruning */);
 
 /**
  * Struct which holds pak file info (version, index offset, hash value).
@@ -94,6 +107,7 @@ struct FPakInfo
 		PakFile_Version_EncryptionKeyGuid = 7,
 		PakFile_Version_FNameBasedCompressionMethod = 8,
 		PakFile_Version_FrozenIndex = 9,
+		PakFile_Version_PathHashIndex = 10,
 
 
 		PakFile_Version_Last,
@@ -109,13 +123,10 @@ struct FPakInfo
 	int64 IndexOffset;
 	/** Size (in bytes) of pak file index. */
 	int64 IndexSize;
-	/** Index SHA1 value. */
+	/** SHA1 of the bytes in the index, used to check for data corruption when loading the index. */
 	FSHAHash IndexHash;
 	/** Flag indicating if the pak index has been encrypted. */
 	uint8 bEncryptedIndex;
-	/** Flag indicating if the pak index has been frozen */
-	// @todo loadtime: we should find a way to unload the index - potentially make two indices, the full one and unloaded one? unclear how, but at least we now have an option to choose per-platform
-	uint8 bIndexIsFrozen;
 	/** Encryption key guid. Empty if we should use the embedded key. */
 	FGuid EncryptionKeyGuid;
 	/** Compression methods used in this pak file (FNames, saved as FStrings) */
@@ -130,7 +141,6 @@ struct FPakInfo
 		, IndexOffset(-1)
 		, IndexSize(0)
 		, bEncryptedIndex(0)
-		, bIndexIsFrozen(0)
 	{
 		// we always put in a NAME_None entry as index 0, so that an uncompressed PakEntry will have CompressionMethodIndex of 0 and can early out easily
 		CompressionMethods.Add(NAME_None);
@@ -146,7 +156,7 @@ struct FPakInfo
 		int64 Size = sizeof(Magic) + sizeof(Version) + sizeof(IndexOffset) + sizeof(IndexSize) + sizeof(IndexHash) + sizeof(bEncryptedIndex);
 		if (InVersion >= PakFile_Version_EncryptionKeyGuid) Size += sizeof(EncryptionKeyGuid);
 		if (InVersion >= PakFile_Version_FNameBasedCompressionMethod) Size += CompressionMethodNameLen * MaxNumCompressionMethods;
-		if (InVersion >= PakFile_Version_FrozenIndex) Size += sizeof(bIndexIsFrozen);
+		if (InVersion >= PakFile_Version_FrozenIndex && InVersion < PakFile_Version_PathHashIndex) Size += sizeof(bool);
 
 		return Size;
 	}
@@ -202,9 +212,14 @@ struct FPakInfo
 			}
 		}
 
-		if (Version >= PakFile_Version_FrozenIndex)
+		if (Version >= PakFile_Version_FrozenIndex && Version < PakFile_Version_PathHashIndex)
 		{
+			bool bIndexIsFrozen = false;
 			Ar << bIndexIsFrozen;
+			if (bIndexIsFrozen)
+			{
+				UE_LOG(LogPakFile, Fatal, TEXT("PakFile was frozen with version FPakInfo::PakFile_Version_FrozenIndex, which is no longer supported. Regenerate Paks."));
+			}
 		}
 
 		if (Version < PakFile_Version_FNameBasedCompressionMethod)
@@ -276,12 +291,10 @@ struct FPakInfo
  */
 struct FPakCompressedBlock
 {
-	DECLARE_TYPE_LAYOUT(FPakCompressedBlock, NonVirtual);
-
 	/** Offset of the start of a compression block. Offset is relative to the start of the compressed chunk data */
-	LAYOUT_FIELD(int64, CompressedStart);
+	int64 CompressedStart;
 	/** Offset of the end of a compression block. This may not align completely with the start of the next block. Offset is relative to the start of the compressed chunk data. */
-	LAYOUT_FIELD(int64, CompressedEnd);
+	int64 CompressedEnd;
 
 	bool operator == (const FPakCompressedBlock& B) const
 	{
@@ -315,30 +328,28 @@ FORCEINLINE FArchive& operator<<(FArchive& Ar, FPakCompressedBlock& Block)
  */
 struct FPakEntry
 {
-	DECLARE_TYPE_LAYOUT(FPakEntry, NonVirtual);
-
 	static const uint8 Flag_None = 0x00;
 	static const uint8 Flag_Encrypted = 0x01;
 	static const uint8 Flag_Deleted = 0x02;
 
 	/** Offset into pak file where the file is stored.*/
-	LAYOUT_FIELD(int64, Offset);
+	int64 Offset;
 	/** Serialized file size. */
-	LAYOUT_FIELD(int64, Size);
+	int64 Size;
 	/** Uncompressed file size. */
-	LAYOUT_FIELD(int64, UncompressedSize);
+	int64 UncompressedSize;
 	/** File SHA1 value. */
-	LAYOUT_ARRAY(uint8, Hash, 20);
+	uint8 Hash[20];
 	/** Array of compression blocks that describe how to decompress this pak entry. */
-	LAYOUT_FIELD(TMemoryImageArray<FPakCompressedBlock>, CompressionBlocks);
+	TArray<FPakCompressedBlock> CompressionBlocks;
 	/** Size of a compressed block in the file. */
-	LAYOUT_FIELD(uint32, CompressionBlockSize);
+	uint32 CompressionBlockSize;
 	/** Index into the compression methods in this pakfile. */
-	LAYOUT_FIELD(uint32, CompressionMethodIndex);
+	uint32 CompressionMethodIndex;
 	/** Pak entry flags. */
-	LAYOUT_FIELD(uint8, Flags);
+	uint8 Flags;
 	/** Flag is set to true when FileHeader has been checked against PakHeader. It is not serialized. */
-	LAYOUT_MUTABLE_FIELD(bool, Verified);
+	mutable bool  Verified;
 
 	/**
 	 * Constructor.
@@ -394,15 +405,8 @@ struct FPakEntry
 	 */
 	bool operator == (const FPakEntry& B) const
 	{
-		// Offsets are not compared here because they're not
-		// serialized with file headers anyway.
-		return Size == B.Size && 
-			UncompressedSize == B.UncompressedSize &&
-			CompressionMethodIndex == B.CompressionMethodIndex &&
-			Flags == B.Flags &&
-			CompressionBlockSize == B.CompressionBlockSize &&
-			FMemory::Memcmp(Hash, B.Hash, sizeof(Hash)) == 0 &&
-			CompressionBlocks == B.CompressionBlocks;
+		return IndexDataEquals(B) &&
+			FMemory::Memcmp(Hash, B.Hash, sizeof(Hash)) == 0;
 	}
 
 	/**
@@ -410,15 +414,20 @@ struct FPakEntry
 	 */
 	bool operator != (const FPakEntry& B) const
 	{
-		// Offsets are not compared here because they're not
-		// serialized with file headers anyway.
-		return Size != B.Size || 
-			UncompressedSize != B.UncompressedSize ||
-			CompressionMethodIndex != B.CompressionMethodIndex ||
-			Flags != B.Flags ||
-			CompressionBlockSize != B.CompressionBlockSize ||
-			FMemory::Memcmp(Hash, B.Hash, sizeof(Hash)) != 0 ||
-			CompressionBlocks != B.CompressionBlocks;
+		return !(*this == B);
+	}
+
+	bool IndexDataEquals(const FPakEntry& B) const
+	{
+		// Offset are only in the Index and so are not compared
+		// Hash is only in the payload and so are not compared
+		// Verified is only in the payload and is mutable and so is not compared
+		return Size == B.Size &&
+			UncompressedSize == B.UncompressedSize &&
+			CompressionMethodIndex == B.CompressionMethodIndex &&
+			Flags == B.Flags &&
+			CompressionBlockSize == B.CompressionBlockSize &&
+			CompressionBlocks == B.CompressionBlocks;
 	}
 
 	/**
@@ -511,30 +520,164 @@ struct FPakEntry
 	static bool VerifyPakEntriesMatch(const FPakEntry& FileEntryA, const FPakEntry& FileEntryB);
 };
 
-/** Pak directory type mapping a filename to a FPakEntry index within the FPakFile.Files array or the MiniPakEntriesOffsets array depending on the bit-encoded state of the FPakEntry structures. */
-typedef TMemoryImageMap<FMemoryImageString, int32> FPakDirectory;
+/**
+ * An identifier for the location of an FPakEntry in an FDirectoryIndex or an FPathHashIndex.
+ * Contains a byte offset into the encoded array of FPakEntry data, an index into the list of unencodable FPakEntries, or a marker indicating invalidity
+ */
+struct FPakEntryLocation
+{
+public:
+	/*
+	 * 0x00000000 - 0x7ffffffe: EncodedOffset from 0 to MaxIndex
+	 * 0x7fffffff: Unused, interpreted as Invalid
+	 * 0x80000000: Invalid
+	 * 0x80000001 - 0xffffffff: FileIndex from MaxIndex to 0
+	*/
+	static CONSTEXPR int32 Invalid = MIN_int32;
+	static CONSTEXPR int32 MaxIndex = MAX_int32 - 1;
+
+	FPakEntryLocation() : Index(Invalid)
+	{
+	}
+	FPakEntryLocation(const FPakEntryLocation& Other) = default;
+	FPakEntryLocation& operator=(const FPakEntryLocation& other) = default;
+
+	static FPakEntryLocation CreateInvalid()
+	{
+		return FPakEntryLocation();
+	}
+
+	static FPakEntryLocation CreateFromOffsetIntoEncoded(int32 Offset)
+	{
+		check(0 <= Offset && Offset <= MaxIndex);
+		return FPakEntryLocation(Offset);
+	}
 
-struct PAKFILE_API FPakFileData
+	static FPakEntryLocation CreateFromListIndex(int32 ListIndex)
+	{
+		check(0 <= ListIndex && ListIndex <= MaxIndex);
+		return FPakEntryLocation(-ListIndex - 1);
+	}
+
+	bool IsInvalid() const
+	{
+		return Index <= Invalid || MaxIndex < Index;
+	}
+
+	bool IsOffsetIntoEncoded() const
+	{
+		return 0 <= Index && Index <= MaxIndex;
+	}
+
+	bool IsListIndex() const
+	{
+		return (-MaxIndex - 1) <= Index && Index <= -1;
+	}
+
+	int32 GetAsOffsetIntoEncoded() const
+	{
+		if (IsOffsetIntoEncoded())
+		{
+			return Index;
+		}
+		else
+		{
+			return -1;
+		}
+	}
+	int32 GetAsListIndex() const
+	{
+		if (IsListIndex())
+		{
+			return -(Index + 1);
+		}
+		else
+		{
+			return -1;
+		}
+	}
+
+	void Serialize(FArchive& Ar)
+	{
+		Ar << Index;
+	}
+
+	bool operator==(const FPakEntryLocation& Other) const
+	{
+		return Index == Other.Index;
+	}
+private:
+	explicit FPakEntryLocation(int32 InIndex) : Index(InIndex)
+	{
+	}
+
+	int32 Index;
+};
+FORCEINLINE FArchive& operator<<(FArchive& Ar, FPakEntryLocation& PakEntryLocation)
 {
-	DECLARE_TYPE_LAYOUT(FPakFileData, NonVirtual);
+	PakEntryLocation.Serialize(Ar);
+	return Ar;
+}
 
+/** Pak directory type mapping a filename to an FPakEntryLocation. */
+typedef TMap<FString, FPakEntryLocation> FPakDirectory;
 
-	/** Mount point. */
-	LAYOUT_FIELD(FMemoryImageString, MountPoint);
-	/** Info on all files stored in pak. */
-	LAYOUT_FIELD(TMemoryImageArray<FPakEntry>, Files);
-	/** Pak Index organized as a map of directories for faster Directory iteration. Completely valid only when bFilenamesRemoved == false, although portions may still be valid after a call to UnloadPakEntryFilenames() while utilizing DirectoryRootsToKeep. */
-	LAYOUT_FIELD((TMemoryImageMap<FMemoryImageString, FPakDirectory>), Index);
+/* Convenience struct for building FPakFile indexes from an enumeration of (Filename,FPakEntry) pairs */
+struct FPakEntryPair
+{
+	FString Filename;
+	FPakEntry Info;
 };
 
 /**
  * Pak file.
  */
-class PAKFILE_API FPakFile : FNoncopyable
+class PAKFILE_API FPakFile : FNoncopyable, public IPakFile
 {
-	friend class FPakPlatformFile;
+public:
+	/** Index data that provides a map from the hash of a Filename to an FPakEntryLocation */
+	typedef TMap<uint64, FPakEntryLocation> FPathHashIndex;
+	/** Index data that keeps an in-memory directoryname/filename tree to map a Filename to an FPakEntryLocation */
+	typedef TMap<FString, FPakDirectory> FDirectoryIndex;
 
-	TUniquePtr<FPakFileData> Data;
+	/** Pak files can share a cache or have their own */
+	enum class ECacheType : uint8
+	{
+		Shared,
+		Individual,
+	};
+
+	/** A ReadLock wrapper that must be used to prevent threading errors around any call to FindPrunedDirectory or internal uses of DirectoryIndex */
+	struct FScopedPakDirectoryIndexAccess
+	{
+		FScopedPakDirectoryIndexAccess(const FPakFile& InPakFile)
+#if ENABLE_PAKFILE_RUNTIME_PRUNING
+			: PakFile(InPakFile)
+			, bRequiresDirectoryIndexLock(PakFile.RequiresDirectoryIndexLock())
+#endif
+		{
+#if ENABLE_PAKFILE_RUNTIME_PRUNING
+			if (bRequiresDirectoryIndexLock)
+			{
+				PakFile.DirectoryIndexLock.ReadLock();
+			}
+#endif
+		}
+#if ENABLE_PAKFILE_RUNTIME_PRUNING
+		~FScopedPakDirectoryIndexAccess()
+		{
+			if (bRequiresDirectoryIndexLock)
+			{
+				PakFile.DirectoryIndexLock.ReadUnlock();
+			}
+		}
+		const FPakFile& PakFile;
+		bool bRequiresDirectoryIndexLock;
+#endif
+	};
+
+private:
+	friend class FPakPlatformFile;
 
 	/** Pak filename. */
 	FString PakFilename;
@@ -549,18 +692,24 @@ class PAKFILE_API FPakFile : FNoncopyable
 	FPakInfo Info;
 	/** Mount point. */
 	FString MountPoint;
-	/** The hash to use when generating a filename hash (CRC) to avoid collisions within the hashed filename space. */
-	uint64 FilenameStartHash;
-	/** An array of 256 + 1 size that represents the starting index of the most significant byte of a hash group within the FilenameHashes array. */
-	uint32* FilenameHashesIndex;
-	/** An array of NumEntries size mapping 1:1 with FilenameHashes and describing the index of the FPakEntry. */
-	int32* FilenameHashesIndices;
-	/** A tightly packed array of filename hashes (CRC) of NumEntries size. */
-	uint64* FilenameHashes;
-	/** A tightly packed array, NumEntries in size, of offsets to the pak entry data within the MiniPakEntries buffer */
-	uint32* MiniPakEntriesOffsets;
-	/** Memory buffer representing the minimal file entry headers, NumEntries in size */
-	uint8* MiniPakEntries;
+	/** Info on all files stored in pak. */
+	TArray<FPakEntry> Files;
+	/** Pak Index organized as a map of directories to support searches by path.  This Index is pruned at runtime of all FileNames and Paths that are not whitelisted in DirectoryIndexKeepFiles */
+	FDirectoryIndex DirectoryIndex;
+#if ENABLE_PAKFILE_RUNTIME_PRUNING
+	/** Temporary-lifetime copy of the Pruned DirectoryIndex; all Pruned files have been removed form this copy.  This copy is used for validation that no queries are missing during runtime, and will be swapped into the DirectoryIndex when Pak Mounting is complete */
+	FDirectoryIndex PrunedDirectoryIndex;
+	/** ReaderWriter lock to guard DirectoryIndex iteration from being interrupted by the swap of PrunedDirectoryIndex */
+	mutable FRWLock DirectoryIndexLock;
+#endif
+
+	/** Index data that provides a map from the hash of a Filename to an FPakEntryLocation */
+	FPathHashIndex PathHashIndex;
+	/* FPakEntries that have been serialized into a compacted format in an array of bytes. */
+	TArray<uint8> EncodedPakEntries;
+	/* The seed passed to the hash function for hashing filenames in this pak.  Differs per pack so that the same filename in different paks has different hashes */
+	uint64 PathHashSeed;
+
 	/** The number of file entries in the pak file */
 	int32 NumEntries;
 	/** Timestamp of this pak file. */
@@ -571,18 +720,28 @@ class PAKFILE_API FPakFile : FNoncopyable
 	bool bSigned;
 	/** True if this pak file is valid and usable. */
 	bool bIsValid;
-	/** True if all filenames in memory for this pak file have been hashed to a 32-bit value. Wildcard traversal is impossible when true. */
-	bool bFilenamesRemoved;
+	/* True if the PathHashIndex has been populated for this PakFile */
+	bool bHasPathHashIndex;
+	/* True if the DirectoryIndex has not been pruned and still contains a Filename for every FPakEntry in this PakFile */
+	bool bHasFullDirectoryIndex;
+#if ENABLE_PAKFILE_RUNTIME_PRUNING
+	/* True if we have a FullDirectoryIndex that we will modify in OptimizeMemoryUsageForMountedPaks and therefore need to guard access with DirectoryIndexLock */
+	bool bWillPruneDirectoryIndex;
+	/* True if the Index of this PakFile was a legacy index that did not have the precomputed Pruned DirectoryIndex and we need to compute it before swapping the Pruned DirectoryIndex*/
+	bool bNeedsLegacyPruning;
+#endif
 	/** ID for the chunk this pakfile is part of. INDEX_NONE if this isn't a pak chunk (derived from filename) */
 	int32 PakchunkIndex;
-	/** Flag to say we tried shrinking pak entries already */
-	bool bAttemptedPakEntryShrink;
-	/** Flag to say we tried unloading pak index filenames already */
-	bool bAttemptedPakFilenameUnload;
 
 	class IMappedFileHandle* MappedFileHandle;
 	FCriticalSection MappedFileHandleCriticalSection;
 
+	/** The type of cache this pak file should have */
+	ECacheType	CacheType;
+	/** The index of this pak file into the cache array, -1 = not initialized */
+	int32		CacheIndex;
+	/** Allow the cache of a pak file to never shrink, should be used with caution, it will burn memory */
+	bool UnderlyingCacheTrimDisabled;
 
 	static inline int32 CDECL CompareFilenameHashes(const void* Left, const void* Right)
 	{
@@ -605,21 +764,35 @@ class PAKFILE_API FPakFile : FNoncopyable
 
 
 public:
-	/** Pak files can share a cache or have their own */
-	enum class ECacheType : uint8
+	// IPakFile interface, for users of PakFiles that cannot have a dependency on this header
+	virtual const FString& PakGetPakFilename() const override
 	{
-		Shared,
-		Individual,
-	};
-private:
-	/** The type of cache this pak file should have */
-	ECacheType	CacheType;
-	/** The index of this pak file into the cache array, -1 = not initialized */
-	int32		CacheIndex;
-	/** Allow the cache of a pak file to never shrink, should be used with caution, it will burn memory */
-	bool UnderlyingCacheTrimDisabled;
+		return PakFilename;
+	}
+
+	virtual bool PakContains(const FString& FullPath) const override
+	{
+		return Find(FullPath, nullptr) == EFindResult::Found;
+	}
+
+	virtual int32 PakGetPakchunkIndex() const override
+	{
+		return PakchunkIndex;
+	}
+
+	virtual void PakVisitPrunedFilenames(IPlatformFile::FDirectoryVisitor& Visitor) const override
+	{
+		for (FFilenameIterator It(*this); It; ++It)
+		{
+			Visitor.Visit(*It.Filename(), false);
+		}
+	}
+
+	virtual const FString& PakGetMountPoint() const override
+	{
+		return MountPoint;
+	}
 
-public:
 
 	void SetUnderlyingCacheTrimDisabled(bool InUnderlyingCacheTrimDisabled) { UnderlyingCacheTrimDisabled = InUnderlyingCacheTrimDisabled; }
 	bool GetUnderlyingCacheTrimDisabled(void) { return UnderlyingCacheTrimDisabled; }
@@ -656,7 +829,7 @@ public:
 	FPakFile(FArchive* Archive);
 #endif
 
-	~FPakFile();
+	virtual ~FPakFile();
 
 	/**
 	 * Checks if the pak file is valid.
@@ -699,9 +872,10 @@ public:
 	 *
 	 * @return Pak index.
 	 */
-	const TMemoryImageMap<FMemoryImageString, FPakDirectory>& GetIndex() const
+	UE_DEPRECATED(4.26, "Use FPrunedFilenameIterator or FPakEntryIterator instead; the Index is now no longer necessarily a DirectoryIndex and in some cases requires a lock around Index use")
+	const TMap<FString, FPakDirectory>& GetIndex() const
 	{
-		return Data->Index;
+		return DirectoryIndex;
 	}
 
 	/**
@@ -712,9 +886,23 @@ public:
 		return NumEntries;
 	}
 
-	void GetFilenames(TArray<FString>& OutFileList) const;
+	UE_DEPRECATED(4.26, "Use GetPrunedFilenames instead")
+	void GetFilenames(TArray<FString>& OutFileList) const
+	{
+		return GetPrunedFilenames(OutFileList);
+	}
 
-	void GetFilenamesInChunk(const TArray<int32>& InChunkIDs, TArray<FString>& OutFileList);
+	/** Returns the FullPath (includes Mount) Filename found in Pruned DirectoryIndex */
+	void GetPrunedFilenames(TArray<FString>& OutFileList) const;
+
+	UE_DEPRECATED(4.26, "Use GetPrunedFilenamesInChunk instead")
+	void GetFilenamesInChunk(const TArray<int32>& InChunkIDs, TArray<FString>& OutFileList) const
+	{
+		GetPrunedFilenamesInChunk(InChunkIDs, OutFileList);
+	}
+
+	/** Returns the RelativePathFromMount Filename for every Filename found in the Pruned DirectoryIndex that points to a PakEntry in the given Chunk */
+	void GetPrunedFilenamesInChunk(const TArray<int32>& InChunkIDs, TArray<FString>& OutFileList) const;
 
 	/**
 	 * Gets shared pak file archive for given thread.
@@ -736,7 +924,7 @@ public:
 		Found,
 		FoundDeleted,
 	};
-	EFindResult Find(const FString& Filename, FPakEntry* OutEntry) const;
+	EFindResult Find(const FString& FullPath, FPakEntry* OutEntry) const;
 
 	/**
 	 * Sets the pak file mount point.
@@ -759,8 +947,18 @@ public:
 		return MountPoint;
 	}
 
+	template <class ContainerType>
+	UE_DEPRECATED(4.26, "Use FindPrunedFilesAtPath instead")
+	void FindFilesAtPath(ContainerType& OutFiles, const TCHAR* InPath, bool bIncludeFiles = true, bool bIncludeDirectories = false, bool bRecursive = false) const
+	{
+		FindPrunedFilesAtPath(OutFiles, InPath, bIncludeFiles, bIncludeDirectories, bRecursive);
+	}
+
 	/**
-	 * Looks for files or directories within the pak file.
+	 * Looks for files or directories within the Pruned DirectoryIndex of the pak file.
+	 * The Pruned DirectoryIndex does not have entries for most Files in the pak; they were removed to save memory.
+	 * A project can specify which FileNames and DirectoryNames can be marked to keep in the DirectoryIndex; see FPakFile::FIndexSettings and FPakFile::PruneDirectoryIndex
+	 * Returned paths are full paths (include the mount point)
 	 *
 	 * @param OutFiles List of files or folder matching search criteria.
 	 * @param InPath Path to look for files or folder at.
@@ -769,7 +967,7 @@ public:
 	 * @param bRecursive If true, sub-folders will also be checked.
 	 */
 	template <class ContainerType>
-	void FindFilesAtPath(ContainerType& OutFiles, const TCHAR* InPath, bool bIncludeFiles = true, bool bIncludeDirectories = false, bool bRecursive = false) const
+	void FindPrunedFilesAtPath(ContainerType& OutFiles, const TCHAR* InPath, bool bIncludeFiles = true, bool bIncludeDirectories = false, bool bRecursive = false) const
 	{
 		// Make sure all directory names end with '/'.
 		FString Directory(InPath);
@@ -778,60 +976,29 @@ public:
 		// Check the specified path is under the mount point of this pak file.
 		// The reverse case (MountPoint StartsWith Directory) is needed to properly handle
 		// pak files that are a subdirectory of the actual directory.
-		if ((Directory.StartsWith(MountPoint)) || (MountPoint.StartsWith(Directory)))
+		if (!Directory.StartsWith(MountPoint) && !MountPoint.StartsWith(Directory))
 		{
-			if (bFilenamesRemoved)
-			{
-				//FPlatformMisc::LowLevelOutputDebugString(*(FString("FindFilesAtPath() used when bFilenamesRemoved == true: ") + InPath));
-			}
+			return;
+		}
 
-			//checkf(!bFilenamesRemoved, TEXT("FPakFile::FindFilesAtPath() can only be used before FPakPlatformFile::UnloadFilenames() is called."));
+		FScopedPakDirectoryIndexAccess ScopeAccess(*this);
+#if ENABLE_PAKFILE_RUNTIME_PRUNING_VALIDATE
+		if (ShouldValidatePrunedDirectory())
+		{
+			TSet<FString> FullFoundFiles, PrunedFoundFiles;
+			FindFilesAtPathInIndex(DirectoryIndex, FullFoundFiles, Directory, bIncludeFiles, bIncludeDirectories, bRecursive);
+			FindFilesAtPathInIndex(PrunedDirectoryIndex, PrunedFoundFiles, Directory, bIncludeFiles, bIncludeDirectories, bRecursive);
+			ValidateDirectorySearch(FullFoundFiles, PrunedFoundFiles, InPath);
 
-			TArray<FString> DirectoriesInPak; // List of all unique directories at path
-			for (TMemoryImageMap<FMemoryImageString, FPakDirectory>::TConstIterator It(Data->Index); It; ++It)
+			for (const FString& FoundFile : FullFoundFiles)
 			{
-				FString PakPath(MountPoint + It.Key());
-				// Check if the file is under the specified path.
-				if (PakPath.StartsWith(Directory))
-				{				
-					if (bRecursive == true)
-					{
-						// Add everything
-						if (bIncludeFiles)
-						{
-							for (FPakDirectory::TConstIterator DirectoryIt(It.Value()); DirectoryIt; ++DirectoryIt)
-							{
-								OutFiles.Add(MountPoint + It.Key() + DirectoryIt.Key());
-							}
-						}
-						if (bIncludeDirectories)
-						{
-							if (Directory != PakPath)
-							{
-								DirectoriesInPak.Add(PakPath);
-							}
-						}
-					}
-					else
-					{
-						int32 SubDirIndex = PakPath.Len() > Directory.Len() ? PakPath.Find(TEXT("/"), ESearchCase::CaseSensitive, ESearchDir::FromStart, Directory.Len() + 1) : INDEX_NONE;
-						// Add files in the specified folder only.
-						if (bIncludeFiles && SubDirIndex == INDEX_NONE)
-						{
-							for (FPakDirectory::TConstIterator DirectoryIt(It.Value()); DirectoryIt; ++DirectoryIt)
-							{
-								OutFiles.Add(MountPoint + It.Key() + DirectoryIt.Key());
-							}
-						}
-						// Add sub-folders in the specified folder only
-						if (bIncludeDirectories && SubDirIndex >= 0)
-						{
-							DirectoriesInPak.AddUnique(PakPath.Left(SubDirIndex + 1));
-						}
-					}
-				}
+				OutFiles.Add(FoundFile);
 			}
-			OutFiles.Append(DirectoriesInPak);
+		}
+		else
+#endif
+		{
+			FindFilesAtPathInIndex(DirectoryIndex, OutFiles, Directory, bIncludeFiles, bIncludeDirectories, bRecursive);
 		}
 	}
 
@@ -841,18 +1008,27 @@ public:
 	 * @param InPath Directory path.
 	 * @return Pointer to a map with directory contents if the directory was found, NULL otherwise.
 	 */
+	UE_DEPRECATED(4.26, "Use FindPrunedDirectory instead and wrap the access with FScopedPakDirectoryIndexAccess; const FPakDirectory is no longer threadsafe.")
 	const FPakDirectory* FindDirectory(const TCHAR* InPath) const
 	{
-		FString Directory(InPath);
-		MakeDirectoryFromPath(Directory);
-		const FPakDirectory* PakDirectory = NULL;
+		return nullptr;
+	}
 
-		// Check the specified path is under the mount point of this pak file.
-		if (Directory.StartsWith(MountPoint))
+	const FPakDirectory* FindPrunedDirectory(const TCHAR* InPath) const
+	{
+		FString RelativePathFromMount;
+		if (!NormalizeDirectoryQuery(InPath, RelativePathFromMount))
 		{
-			PakDirectory = Data->Index.Find(Directory.Mid(MountPoint.Len()));
+			return nullptr;
 		}
-		return PakDirectory;
+
+		return FindPrunedDirectoryInternal(RelativePathFromMount);
+	}
+
+	UE_DEPRECATED(4.26, "Use DirectoryExistsInPruned instead.")
+	bool DirectoryExists(const TCHAR* InPath) const
+	{
+		return DirectoryExistsInPruned(InPath);
 	}
 
 	/**
@@ -861,11 +1037,18 @@ public:
 	 * @param InPath Directory path.
 	 * @return true if the given path exists in pak file, false otherwise.
 	 */
-	bool DirectoryExists(const TCHAR* InPath) const
+	bool DirectoryExistsInPruned(const TCHAR* InPath) const
 	{
-		return !!FindDirectory(InPath);
+		FString RelativePathFromMount;
+		if (!NormalizeDirectoryQuery(InPath, RelativePathFromMount))
+		{
+			return false;
+		}
+
+		FScopedPakDirectoryIndexAccess ScopeAccess(*this);
+		return FindPrunedDirectoryInternal(RelativePathFromMount) != nullptr;
 	}
-	
+
 	/**
 	 * Checks the validity of the pak data by reading out the data for every file in the pak
 	 *
@@ -873,103 +1056,242 @@ public:
 	 */
 	bool Check();
 
-	/** Iterator class used to iterate over all files in pak. */
-	class FFileIterator
+	/** Base functionality for iterating over the DirectoryIndex. */
+	class FBaseIterator
 	{
+	private:
 		/** Owner pak file. */
 		const FPakFile& PakFile;
-		/** Index iterator. */
-		TMemoryImageMap<FMemoryImageString, FPakDirectory>::TConstIterator IndexIt;
-		/** Directory iterator. */
+		/** Outer iterator over Directories when using the FDirectoryIndex. */
+		FDirectoryIndex::TConstIterator DirectoryIndexIt;
+		/** Inner iterator over Files when using the FDirectoryIndex. */
 		FPakDirectory::TConstIterator DirectoryIt;
-		/** The cached filename for return in Filename() */
-		FString CachedFilename;
-		/** Whether to include delete records in the iteration */
+		/** Iterator when using the FPathHashIndex. */
+		FPathHashIndex::TConstIterator PathHashIt;
+		/** The cached filename for return in Filename(). */
+		mutable FString CachedFilename;
+		/* The PakEntry for return in Info */
+		mutable FPakEntry PakEntry;
+		/* Whether to use the PathHashIndex. */
+		bool bUsePathHash;
+		/** Whether to include delete records in the iteration. */
 		bool bIncludeDeleted;
+#if ENABLE_PAKFILE_RUNTIME_PRUNING
+		/** Whether this iterator needs to ReadLock and ReadUnlock due to use of the DirectoryIndex */
+		bool bRequiresDirectoryIndexLock;
+#endif
 
 	public:
-		/**
-		 * Constructor.
-		 *
-		 * @param InPakFile Pak file to iterate.
-		 */
-		FFileIterator(const FPakFile& InPakFile, bool bInIncludeDeleted = false )
-			: PakFile(InPakFile)
-			, IndexIt(PakFile.GetIndex())
-			, DirectoryIt((IndexIt ? FPakDirectory::TConstIterator(IndexIt.Value()): FPakDirectory()))
-			, bIncludeDeleted(bInIncludeDeleted)
+		FBaseIterator& operator++()
 		{
+			if (bUsePathHash)
+			{
+				++PathHashIt;
+			}
+			else
+			{
+				++DirectoryIt;
+			}
 			AdvanceToValid();
-			UpdateCachedFilename();
-		}
-
-		FFileIterator& operator++()		
-		{ 
-			// Continue with the next file
-			++DirectoryIt;
-			AdvanceToValid();
-			UpdateCachedFilename();
-			return *this; 
+			return *this;
 		}
 
 		/** conversion to "bool" returning true if the iterator is valid. */
 		FORCEINLINE explicit operator bool() const
-		{ 
-			return !!IndexIt; 
+		{
+			if (bUsePathHash)
+			{
+				return !!PathHashIt;
+			}
+			else
+			{
+				return !!DirectoryIndexIt;
+			}
 		}
+
 		/** inverse of the "bool" operator */
 		FORCEINLINE bool operator !() const
 		{
 			return !(bool)*this;
 		}
 
-		const FString& Filename() const		{ return CachedFilename; }
-		const FPakEntry& Info() const	{ return PakFile.Data->Files[DirectoryIt.Value()]; }
-		const int32 GetIndexInPakFile() const { return DirectoryIt.Value(); }
+		/** Return the FPakEntry. Invalid to call unless the iterator is currently valid. */
+		const FPakEntry& Info() const
+		{
+			PakFile.GetPakEntry(GetPakEntryIndex(), &PakEntry);
+			return PakEntry;
+		}
 
-	private:
-		void AdvanceToValid()
+		bool HasFilename() const
+		{
+			return !bUsePathHash;
+		}
+
+	protected:
+		FBaseIterator(const FPakFile& InPakFile, bool bInIncludeDeleted, bool bInUsePathHash)
+			: PakFile(InPakFile)
+			, DirectoryIndexIt(FDirectoryIndex())
+			, DirectoryIt(FPakDirectory())
+			, PathHashIt(PakFile.PathHashIndex)
+			, bUsePathHash(bInUsePathHash)
+			, bIncludeDeleted(bInIncludeDeleted)
+#if ENABLE_PAKFILE_RUNTIME_PRUNING
+			, bRequiresDirectoryIndexLock(false)
+#endif
 		{
-			SkipDeletedIfRequired();
-			while (!DirectoryIt && IndexIt)
+			if (bUsePathHash)
+			{
+				check(PakFile.bHasPathHashIndex);
+			}
+			else
 			{
-				// No more files in the current directory, jump to the next one.
-				++IndexIt;
-				if (IndexIt)
+#if ENABLE_PAKFILE_RUNTIME_PRUNING
+				bRequiresDirectoryIndexLock = PakFile.RequiresDirectoryIndexLock();
+				if (bRequiresDirectoryIndexLock)
+				{
+					PakFile.DirectoryIndexLock.ReadLock();
+				}
+#endif
+				DirectoryIndexIt.~TConstIterator();
+				new (&DirectoryIndexIt) FDirectoryIndex::TConstIterator(InPakFile.DirectoryIndex);
+				if (DirectoryIndexIt)
 				{
-					// No need to check if there's files in the current directory. If a directory
-					// exists in the index it is always non-empty.
 					DirectoryIt.~TConstIterator();
-					new(&DirectoryIt) FPakDirectory::TConstIterator(IndexIt.Value());
-					SkipDeletedIfRequired();
+					new (&DirectoryIt) FPakDirectory::TConstIterator(DirectoryIndexIt.Value());
+				}
+			}
+
+			AdvanceToValid();
+		}
+
+#if ENABLE_PAKFILE_RUNTIME_PRUNING
+		~FBaseIterator()
+		{
+			if (bRequiresDirectoryIndexLock)
+			{
+				PakFile.DirectoryIndexLock.ReadUnlock();
+			}
+		}
+#endif
+
+		/** Return the current filename, as the RelativePath from the MountPoint.  Only available when using the FDirectoryIndex, otherwise always returns empty string. Invalid to call unless the iterator is currently valid. */
+		const FString& Filename() const
+		{
+			if (bUsePathHash)
+			{
+				// Filenames are not supported, CachedFilename is always empty 
+			}
+			else
+			{
+				checkf(DirectoryIndexIt && DirectoryIt, TEXT("It is not legal to call Filename() on an invalid iterator"));
+				if (CachedFilename.IsEmpty())
+				{
+					CachedFilename = PakPathCombine(DirectoryIndexIt.Key(), DirectoryIt.Key());
 				}
 			}
+			return CachedFilename;
 		}
 
-		FORCEINLINE void UpdateCachedFilename()
+		/** Return the arbitrary index of the iteration. Invalid to call unless the iterator is currently valid. */
+		const FPakEntryLocation GetPakEntryIndex() const
 		{
-			if (!!IndexIt && !!DirectoryIt)
+			if (bUsePathHash)
 			{
-				CachedFilename = FString(IndexIt.Key()) + FString(DirectoryIt.Key());
+				return PathHashIt.Value();
 			}
 			else
 			{
-				CachedFilename.Empty();
+				return DirectoryIt.Value();
 			}
 		}
 
-		FORCEINLINE void SkipDeletedIfRequired()
+	private:
+
+		/* Skips over deleted records and moves to the next Directory in the DirectoryIndex when necessary. */
+		FORCEINLINE void AdvanceToValid()
 		{
-			if (!bIncludeDeleted)
+			if (bUsePathHash)
 			{
-				while (DirectoryIt && Info().IsDeleteRecord())
+				while (PathHashIt && !bIncludeDeleted && Info().IsDeleteRecord())
 				{
-					++DirectoryIt;
+					++PathHashIt;
 				}
 			}
+			else
+			{
+				while (DirectoryIndexIt && (!DirectoryIt || (!bIncludeDeleted && Info().IsDeleteRecord())))
+				{
+					if (DirectoryIt)
+					{
+						++DirectoryIt;
+					}
+					else
+					{
+						// No more files in the current directory, jump to the next one.
+						++DirectoryIndexIt;
+						if (DirectoryIndexIt)
+						{
+							DirectoryIt.~TConstIterator();
+							new (&DirectoryIt) FPakDirectory::TConstIterator(DirectoryIndexIt.Value());
+						}
+					}
+				}
+				CachedFilename.Reset();
+			}
 		}
 	};
 
+	/** Iterator class for every FPakEntry in the FPakFile, but does not provide filenames unless the PakFile has an unpruned DirectoryIndex. */
+	class FPakEntryIterator : public FBaseIterator
+	{
+	public:
+		FPakEntryIterator(const FPakFile& InPakFile, bool bInIncludeDeleted = false)
+			: FBaseIterator(InPakFile, bInIncludeDeleted, !InPakFile.bHasFullDirectoryIndex /* bUsePathHash */)
+		{
+		}
+
+		const FString* TryGetFilename() const
+		{
+			if (HasFilename())
+			{
+				return &Filename();
+			}
+			else
+			{
+				return nullptr;
+			}
+		}
+	};
+
+	/** Iterator class used to iterate over just the files in the pak for which we have filenames. */
+	class FFilenameIterator : public FBaseIterator
+	{
+	public:
+		/**
+		 * Constructor.
+		 *
+		 * @param InPakFile Pak file to iterate.
+		 */
+		FFilenameIterator(const FPakFile& InPakFile, bool bInIncludeDeleted = false)
+			: FBaseIterator(InPakFile, bInIncludeDeleted, false /* bUsePathHash */)
+		{
+		}
+
+		using FBaseIterator::Filename;
+	};
+
+	class FFileIterator : public FFilenameIterator
+	{
+	public:
+		UE_DEPRECATED(4.26, "FFileIterator is deprecated, use FFilenameIterator instead.  Note that FFilenameIterator will only iterate over the DirectoryIndexKeepFiles entries.")
+		FFileIterator(const FPakFile& InPakFile, bool bInIncludeDeleted = false)
+			: FFilenameIterator(InPakFile, bInIncludeDeleted)
+		{
+		}
+
+		using FFilenameIterator::Filename;
+	};
+
 	/**
 	 * Gets this pak file info.
 	 *
@@ -989,204 +1311,334 @@ public:
 	{
 		return Timestamp;
 	}
-	
+
 	/**
-	 * Returns whether the pak filenames are still resident in memory.
+	 * Returns whether filenames currently exist in the DirectoryIndex for all files in the Pak.
 	 *
-	 * @return true if filenames are present, false otherwise.
+	 * @return true if filenames are present for all FPakEntry, false otherwise.
 	 */
 	bool HasFilenames() const
 	{
-		return !bFilenamesRemoved;
+		return bHasFullDirectoryIndex;
 	}
 
-	/**
-	 * Saves memory by hashing the filenames, if possible. After this process,
-	 * wildcard scanning of pak entries can no longer be performed. Returns TRUE
-	 * if the process successfully unloaded filenames from this pak
+	// FPakFile helper functions shared between the runtime and UnrealPak.exe
+	/*
+	 * Given a FPakEntry from the index, seek to the payload and read the hash of the payload out of the payload entry
+	 * Warning: Slow function, do not use in performance critical operations
 	 *
-	 * @param CrossPakCollisionChecker A map of hash->fileentry records encountered during filename unloading on other pak files. Used to detect collisions with entries in other pak files.
-	 * @param DirectoryRootsToKeep An array of strings in wildcard format that specify whole directory structures of filenames to keep in memory for directory iteration to work.
-	 * @param bAllowRetries If a collision is encountered, change the intial seed and try again a fixed number of times before failing
-	 */
-	bool UnloadPakEntryFilenames(TMap<uint64, FPakEntry>& CrossPakCollisionChecker, TArray<FString>* DirectoryRootsToKeep = nullptr, bool bAllowRetries = true);
-
-	/**
-	 * Lower memory usage by bit-encoding the pak file entry information.
+	 * @param PakEntry the FPakEntry from the index, which has the Offset to read to
+	 * @param OutBuffer a buffer at least sizeof(FPakEntry::Hash) in size, into which the hash will be copied
 	 */
-	bool ShrinkPakEntriesMemoryUsage();
-
-	/**
-	 * Returns whether the pak files list has been shrunk or not
-	 */
-	bool HasShrunkPakEntries() const
+	void ReadHashFromPayload(const FPakEntry& PakEntry, uint8* OutBuffer)
 	{
-		return bAttemptedPakEntryShrink;
+		if (PakEntry.IsDeleteRecord())
+		{
+			FMemory::Memset(OutBuffer, 0, sizeof(FPakEntry::Hash));
+		}
+		else
+		{
+			FArchive* Reader = GetSharedReader(nullptr);
+			Reader->Seek(PakEntry.Offset);
+			FPakEntry SerializedEntry;
+			SerializedEntry.Serialize(*Reader, GetInfo().Version);
+			FMemory::Memcpy(OutBuffer, &SerializedEntry.Hash, sizeof(SerializedEntry.Hash));
+		}
 	}
 
-private:
+	/** Hash the given full-path filename using the hash function used by FPakFiles, with the given FPakFile-specific seed */
+	static uint64 HashPath(const TCHAR* RelativePathFromMount, uint64 Seed);
 
-	/**
-	 * Initializes the pak file.
-	 */
-	void Initialize(FArchive* Reader, bool bLoadIndex = true);
+	/** Read a list of (Filename, FPakEntry) pairs from a provided enumeration, attempt to encode each one,
+	  * store each one in the appropriate given encoded and/or unencoded array, and populate the given
+	  * Directories to map each filename to the location for the FPakEntry
 
-	/**
-	 * Loads and initializes pak file index.
-	 */
-	void LoadIndex(FArchive* Reader);
+	  * @param InNumEntries How many entries will be read from ReadNextEntryFunction
+	  * @param ReadNextEntryFunction Callback called repeatedly to enumerate the (Filename,FPakEntry) pairs to be encoded
+	  * @param InPakFilename Filename for the pak containing the files, used to create the hashseed for the given pak
+	  * @param InPakInfo PakInfo for the given pak, used for serialization flags
+	  * @param MountPoint Directory into which the pak will be mounted, used to create the Directory and PathHash indexes
+	  * @param OutNumEncodedEntries How many entries were written to the bytes in OutEncodedPakEntries
+	  * @param OutNumDeletedEntries How many entries were skipped and not stored because the input FPakEntry was a Delete record
+	  * @param OutPathHashSeed optional out param to get a copy of the pakfile-specific hashseed
+	  * @param OutDirectoryIndex optional output FDirectoryIndex
+	  * @param OutPathHashIndex optional output FPathHashIndex
+	  * @param OutEncodedPakEntries array of bytes into which the encoded FPakEntries are stored.  Values in OutDirectoryIndex and OutPathHashIndex can be offsets into this array indicated the start point for the encoding of the given FPakEntry
+	  * @param OutNonEncodableEntries A list of all the FPakEntries that could not be encoded.  Values in OutDirectoryIndex and OutPathHashIndex can be indices into this list.
+	  * @param InOutCollisionDetection Optional parameter to detect hash collisions.  If present, each hashed filename will be check()'d for a collision against a different filename in InOutCollisionDetection, and will be added into InOutCollisionDetection
+	  */
+	typedef TFunction<FPakEntryPair & ()> ReadNextEntryFunction;
+	static void EncodePakEntriesIntoIndex(int32 InNumEntries, const ReadNextEntryFunction& InReadNextEntry, const TCHAR* InPakFilename, const FPakInfo& InPakInfo, const FString& MountPoint,
+		int32& OutNumEncodedEntries, int32& OutNumDeletedEntries, uint64* OutPathHashSeed,
+		FDirectoryIndex* OutDirectoryIndex, FPathHashIndex* OutPathHashIndex, TArray<uint8>& OutEncodedPakEntries, TArray<FPakEntry>& OutNonEncodableEntries, TMap<uint64, FString>* InOutCollisionDetection);
+
+	/** Lookup the FPakEntryLocation stored in the given PathHashIndex, return nullptr if not found */
+	static const FPakEntryLocation* FindLocationFromIndex(const FString& FullPath, const FString& MountPoint, const FPathHashIndex& PathHashIndex, uint64 PathHashSeed);
+
+	/** Lookup the FPakEntryLocation stored in the given DirectoryIndex, return nullptr if not found */
+	static const FPakEntryLocation* FindLocationFromIndex(const FString& FullPath, const FString& MountPoint, const FDirectoryIndex& DirectoryIndex);
 
 	/**
-	 * Manually add a file to a pak file,
-	 */
-	void AddSpecialFile(FPakEntry Entry, const FString& Filename);
+	  * Returns the FPakEntry pointed to by the given FPakEntryLocation inside the given EncodedPakEntries or Files
+	  * Can return Found or Deleted; if the FPakEntryLocation is invalid this function assumes the FPakEntry exists in this pack but as a deleted file
+	  * If OutEntry is non-null, populates it with a copy of the FPakEntry found, or sets it to
+	  * an FPakEntry with SetDeleteRecord(true) if not found
+	  */
+	static EFindResult GetPakEntry(const FPakEntryLocation& FPakEntryLocation, FPakEntry* OutEntry, const TArray<uint8>& EncodedPakEntries, const TArray<FPakEntry>& Files, const FPakInfo& Info);
 
 	/**
-	 * Decodes a bit-encoded pak entry.
+	 * Given a directory index, remove entries from it that are directed by ini to not have filenames kept at runtime.
 	 *
-	 * @param Filename File to find.
-	 * @param OutEntry The optional address of an FPakEntry instance where the found file information should be stored. Pass NULL to only check for file existence.
-	 * @return Returns true if the file was found, false otherwise.
+	 * InOutDirectoryIndex - The full index from which to potentially remove entries
+	 * OutDirectoryIndex - If null, InOutDirectoryIndex will have pruned entries removed.  If non-null, InOutDirectoryIndex will not be modified, and PrunedDirectoryIndex will have kept values added.
+	 * MountPoint The mount point for the pak containing the index, used to provide the fullpath for filenames in the DirectoryIndex for comparison against paths in ini
 	 */
-	bool DecodePakEntry(const uint8* SourcePtr, FPakEntry* OutEntry) const
-	{
-		// Grab the big bitfield value:
-		// Bit 31 = Offset 32-bit safe?
-		// Bit 30 = Uncompressed size 32-bit safe?
-		// Bit 29 = Size 32-bit safe?
-		// Bits 28-23 = Compression method
-		// Bit 22 = Encrypted
-		// Bits 21-6 = Compression blocks count
-		// Bits 5-0 = Compression block size
-		uint32 Value = *(uint32*)SourcePtr;
-		SourcePtr += sizeof(uint32);
-
-		// Filter out the CompressionMethod.
-		OutEntry->CompressionMethodIndex = (Value >> 23) & 0x3f;
-
-		// Test for 32-bit safe values. Grab it, or memcpy the 64-bit value
-		// to avoid alignment exceptions on platforms requiring 64-bit alignment
-		// for 64-bit variables.
-		//
-		// Read the Offset.
-		bool bIsOffset32BitSafe = (Value & (1 << 31)) != 0;
-		if (bIsOffset32BitSafe)
+	static void PruneDirectoryIndex(FDirectoryIndex& InOutDirectoryIndex, FDirectoryIndex* PrunedDirectoryIndex, const FString& MountPoint);
+
+	/* Helper function to modify the given string to append '/' at the end of path to normalize directory names for hash and string compares */
+	static void MakeDirectoryFromPath(FString& Path)
+	{
+		if (Path.Len() > 0 && Path[Path.Len() - 1] != '/')
 		{
-			OutEntry->Offset = *(uint32*)SourcePtr;
-			SourcePtr += sizeof(uint32);
+			Path += TEXT("/");
 		}
-		else
+	}
+	/* Helper function to check that the given string is in our directory format (ends with '/') */
+	static bool IsPathInDirectoryFormat(const FString& Path)
+	{
+		return Path.Len() > 0 && Path[Path.Len() - 1] == TEXT('/');
+	}
+
+	/* Helper function to join two path strings that are in the PakPath format */
+	static FString PakPathCombine(const FString& Parent, const FString& Child)
+	{
+		// Our paths are different than FPaths, because our dirs / at the end, and "/" is the relative path to the mountdirectory and should be mapped to the empty string when joining
+		check(Parent.Len() > 0 && Parent[Parent.Len() - 1] == TEXT('/'));
+		if (Parent.Len() == 1)
 		{
-			FMemory::Memcpy(&OutEntry->Offset, SourcePtr, sizeof(int64));
-			SourcePtr += sizeof(int64);
+			return Child;
 		}
-
-		// Read the UncompressedSize.
-		bool bIsUncompressedSize32BitSafe = (Value & (1 << 30)) != 0;
-		if (bIsUncompressedSize32BitSafe)
+		else if (Child.Len() == 1 && Child[0] == TEXT('/'))
 		{
-			OutEntry->UncompressedSize = *(uint32*)SourcePtr;
-			SourcePtr += sizeof(uint32);
+			return Parent;
 		}
 		else
 		{
-			FMemory::Memcpy(&OutEntry->UncompressedSize, SourcePtr, sizeof(int64));
-			SourcePtr += sizeof(int64);
+			check(Child.Len() == 0 || Child[0] != TEXT('/'));
+			return Parent + Child;
 		}
+	}
 
-		// Fill in the Size.
-		if (OutEntry->CompressionMethodIndex != 0)
+	/** Helper function to split a PakDirectoryIndex-Formatted PathName into its PakDirectoryIndex-Formatted parent directory and the CleanFileName */
+	static bool SplitPathInline(FString& InOutPath, FString& OutFilename)
+	{
+		// FPaths::GetPath doesn't handle our / at the end of directories, so we have to do string manipulation ourselves
+		// The manipulation is less complicated than GetPath deals with, since we have normalized/path/strings, we have relative paths only, and we don't care about extensions
+		if (InOutPath.Len() == 0)
+		{
+			check(false); // Filenames should have non-zero length, and the minimum directory length is 1 (The root directory is written as "/")
+			return false;
+		}
+		else if (InOutPath.Len() == 1)
 		{
-			// Size is only present if compression is applied.
-			bool bIsSize32BitSafe = (Value & (1 << 29)) != 0;
-			if (bIsSize32BitSafe)
+			if (InOutPath[0] == TEXT('/'))
 			{
-				OutEntry->Size = *(uint32*)SourcePtr;
-				SourcePtr += sizeof(uint32);
+				// The root directory; it has no parent.
+				OutFilename.Empty();
+				return false;
 			}
 			else
 			{
-				FMemory::Memcpy(&OutEntry->Size, SourcePtr, sizeof(int64));
-				SourcePtr += sizeof(int64);
+				// A relative one-character path with no /; this is a direct child of in the root directory
+				OutFilename = TEXT("/");
+				Swap(OutFilename, InOutPath);
+				return true;
 			}
 		}
 		else
 		{
-			// The Size is the same thing as the UncompressedSize when
-			// CompressionMethod == COMPRESS_None.
-			OutEntry->Size = OutEntry->UncompressedSize;
+			if (InOutPath[InOutPath.Len() - 1] == TEXT('/'))
+			{
+				// The input was a Directory; remove the trailing / since we don't keep those on the CleanFilename
+				InOutPath.LeftChopInline(1, false /* bAllowShrinking */);
+			}
+
+			int32 Offset = 0;
+			if (InOutPath.FindLastChar(TEXT('/'), Offset))
+			{
+				int32 FilenameStart = Offset + 1;
+				OutFilename = InOutPath.Mid(FilenameStart);
+				InOutPath.LeftInline(FilenameStart, false /* bAllowShrinking */); // The Parent Directory keeps the / at the end
+			}
+			else
+			{
+				// A relative path with no /; this is a direct child of in the root directory
+				OutFilename = TEXT("/");
+				Swap(OutFilename, InOutPath);
+			}
+			return true;
 		}
+	}
 
-		// Filter the encrypted flag.
-		OutEntry->SetEncrypted((Value & (1 << 22)) != 0);
+	/* Returns the global,const flag for whether the current process is allowing PakFiles to keep their entire DirectoryIndex (if it exists in the PakFile on disk) rather than pruning it */
+	static bool IsPakKeepFullDirectory();
 
-		// This should clear out any excess CompressionBlocks that may be valid in the user's
-		// passed in entry.
-		uint32 CompressionBlocksCount = (Value >> 6) & 0xffff;
-		OutEntry->CompressionBlocks.Empty(CompressionBlocksCount);
-		OutEntry->CompressionBlocks.SetNum(CompressionBlocksCount);
+	/* Returns the global,const flag for whether UnrealPak should write a copy of the full PathHashIndex and Pruned DirectoryIndex to the PakFile */
+	static bool IsPakWritePathHashIndex();
 
-		// Filter the compression block size or use the UncompressedSize if less that 64k.
-		OutEntry->CompressionBlockSize = 0;
-		if (CompressionBlocksCount > 0)
-		{
-			OutEntry->CompressionBlockSize = OutEntry->UncompressedSize < 65536 ? (uint32)OutEntry->UncompressedSize : ((Value & 0x3f) << 11);
-		}
+	/* Returns the global,const flag for whether UnrealPak should write a copy of the full DirectoryIndex to the PakFile */
+	static bool IsPakWriteFullDirectoryIndex();
 
-		// Set Verified to true to avoid have a synchronous open fail comparing FPakEntry structures.
-		OutEntry->Verified = true;
+private:
 
-		// Set bDeleteRecord to false, because it obviously isn't deleted if we are here.
-		OutEntry->SetDeleteRecord(false);
+	/**
+	 * Initializes the pak file.
+	 */
+	void Initialize(FArchive* Reader, bool bLoadIndex = true);
 
-		// Base offset to the compressed data
-		int64 BaseOffset = Info.HasRelativeCompressedChunkOffsets() ? 0 : OutEntry->Offset;
+	/**
+	 * Loads and initializes pak file index.
+	 */
+	void LoadIndex(FArchive* Reader);
 
-		// Handle building of the CompressionBlocks array.
-		if (OutEntry->CompressionBlocks.Num() == 1 && !OutEntry->IsEncrypted())
-		{
-			// If the number of CompressionBlocks is 1, we didn't store any extra information.
-			// Derive what we can from the entry's file offset and size.
-			FPakCompressedBlock& CompressedBlock = OutEntry->CompressionBlocks[0];
-			CompressedBlock.CompressedStart = BaseOffset + OutEntry->GetSerializedSize(Info.Version);
-			CompressedBlock.CompressedEnd = CompressedBlock.CompressedStart + OutEntry->Size;
-		}
-		else if (OutEntry->CompressionBlocks.Num() > 0)
-		{
-			// Get the right pointer to start copying the CompressionBlocks information from.
-			uint32* CompressionBlockSizePtr = (uint32*)SourcePtr;
+	/**
+	  * Returns the FPakEntry pointed to by the given FPakEntryLocation, forwards to the static GetPakEntry with data from *this
+	  */
+	EFindResult GetPakEntry(const FPakEntryLocation& FPakEntryLocation, FPakEntry* OutEntry) const;
 
-			// Alignment of the compressed blocks
-			uint64 CompressedBlockAlignment = OutEntry->IsEncrypted() ? FAES::AESBlockSize : 1;
+	/** Helper class to read IndexSettings from project delegate and commandline */
+	struct FIndexSettings;
+	static FIndexSettings& GetIndexSettings();
 
-			// CompressedBlockOffset is the starting offset. Everything else can be derived from there.
-			int64 CompressedBlockOffset = BaseOffset + OutEntry->GetSerializedSize(Info.Version);
-			for (int CompressionBlockIndex = 0; CompressionBlockIndex < OutEntry->CompressionBlocks.Num(); ++CompressionBlockIndex)
-			{
-				FPakCompressedBlock& CompressedBlock = OutEntry->CompressionBlocks[CompressionBlockIndex];
-				CompressedBlock.CompressedStart = CompressedBlockOffset;
-				CompressedBlock.CompressedEnd = CompressedBlockOffset + *CompressionBlockSizePtr++;
-				CompressedBlockOffset += Align(CompressedBlock.CompressedEnd - CompressedBlock.CompressedStart, CompressedBlockAlignment);
-			}
-		}
+	/**
+	  * Returns the global,const flag for whether the current process should run directory queries on both the DirectoryIndex and the Pruned DirectoryIndex and log an error if they don't match.
+	  * Validation only occurs until the first call to OptimizeMemoryUsageForMountedPaks, after which the Full DirectoryIndex is dropped and there is nothing left to Validate
+	  * Has the same effect as IsPakDelayPruning, plus the addition of the error for any mismatches.
+	  */
+	static bool IsPakValidatePruning();
+	/**
+	 * Returns the global,const flag for whether the current process should keep a copy of the Full DirectoryIndex around until OptimizeMemoryUsageForMountedPaks is called, so that systems can run
+	 * directory queries against the full index until then.
+	 * Note that validation will still occur if IsPakValidatePruning is true.
+	 */
+	static bool IsPakDelayPruning();
 
-		return true;
-	}
+#if ENABLE_PAKFILE_RUNTIME_PRUNING
+	/** Global flag for whether a Pak has indicated it needs Pruning */
+	static bool bSomePakNeedsPruning;
+#endif
 
-public:
+	/**
+	  * Returns whether read accesses against the DirectoryIndex need to be guarded using this->DirectoryIndexLock.
+	  * Locking is not required if the pak is not going to be pruned or already has been; the DirectoryIndex is immutable after that point, and we can get a performance benefit by skipping the lock.
+	  */
+	bool RequiresDirectoryIndexLock() const;
 
 	/**
-	 * Helper function to append '/' at the end of path.
-	 *
-	 * @param Path - path to convert in place to directory.
+	 * Returns whether the current Process IsPakValidatePruning and this PakFile has a Full DirectoryIndex and Pruned DirectoryIndex to validate.
 	 */
-	static void MakeDirectoryFromPath(FString& Path)
+	bool ShouldValidatePrunedDirectory() const;
+
+	/**
+	  * Add the given (Filename,FPakEntryLocation) value into the provided indexes
+	  *
+	  * @param Filename The filename to add 
+	  * @param EntryLocation The FPakEntryLocation to add
+	  * @param MountPoint The mount point of the pakfile containing the FPakEntryLocation, used to create the filename in the DirectoryIndex
+	  * @param PathHashSeed the pakfile-specific seed for the hash of the path in the PathHasIndex
+	  * @param DirectoryIndex Optional FDirectoryIndex into which to insert the (Filename, FPakEntryLocation)
+	  * @param PathHashIndex Optional FPathHashIndex into which to insert the (Filename, FPakEntryLocation)
+	  * @param InOutCollisionDetection Optional parameter to detect hash collisions.  If present, the hashed filename will be check()'d for a collision against a different filename in InOutCollisionDetection, and will be added into InOutCollisionDetection
+	  */
+	static void AddEntryToIndex(const FString& Filename, const FPakEntryLocation& EntryLocation, const FString& MountPoint, uint64 PathHashSeed,
+		FDirectoryIndex* DirectoryIndex, FPathHashIndex* PathHashIndex, TMap<uint64, FString>* CollisionDetection);
+
+	/* Encodes a pak entry as an array of bytes into the given archive.  Returns true if encoding succeeded.  If encoding did not succeed, caller will need to store the InPakEntry in an unencoded list */
+	static bool EncodePakEntry(FArchive& Ar, const FPakEntry& InPakEntry, const FPakInfo& InInfo);
+
+	/* Decodes a bit-encoded pak entry from a pointer to the start of its encoded bytes into the given OutEntry */
+	static void DecodePakEntry(const uint8* SourcePtr, FPakEntry& OutEntry, const FPakInfo& InInfo);
+
+	/* Internal index loading function that returns false if index loading fails due to an intermittent IO error. Allows LoadIndex to retry or throw a fatal as required */
+	bool LoadIndexInternal(FArchive* Reader);
+
+	/* Legacy index loading function for PakFiles saved before FPakInfo::PakFile_Version_PathHashIndex */
+	bool LoadLegacyIndex(FArchive* Reader);
+
+	/* Helper function for LoadIndexInternal; each array of Index bytes read from the file needs to be independently decrypted and checked for corruption */
+	bool DecryptAndValidateIndex(FArchive* Reader, TArray<uint8>& IndexData, FSHAHash& InExpectedHash, FSHAHash& OutActualHash);
+
+	/* Manually add a file to a pak file */
+	void AddSpecialFile(const FPakEntry& Entry, const FString& Filename);
+
+	/**
+	 * Search the given FDirectoryIndex for all files under the given Directory.  Helper for FindFilesAtPath, called separately on the DirectoryIndex or Pruned DirectoryIndex. Does not use
+	 * FScopedPakDirectoryIndexAccess internally; caller is responsible for calling from within a lock.
+	 * Returned paths are full paths (include the mount point)
+	 */
+	template <class ContainerType>
+	void FindFilesAtPathInIndex(const FDirectoryIndex& TargetIndex, ContainerType& OutFiles, const FString& Directory, bool bIncludeFiles = true, bool bIncludeDirectories = false, bool bRecursive = false) const
 	{
-		if (Path.Len() > 0 && Path[Path.Len() - 1] != '/')
+		TArray<FString> DirectoriesInPak; // List of all unique directories at path
+		for (TMap<FString, FPakDirectory>::TConstIterator It(TargetIndex); It; ++It)
 		{
-			Path += TEXT("/");
+			const FString PakPath(PakPathCombine(MountPoint, It.Key()));
+			// Check if the file is under the specified path.
+			if (PakPath.StartsWith(Directory))
+			{
+				if (bRecursive == true)
+				{
+					// Add everything
+					if (bIncludeFiles)
+					{
+						for (FPakDirectory::TConstIterator DirectoryIt(It.Value()); DirectoryIt; ++DirectoryIt)
+						{
+							OutFiles.Add(PakPathCombine(PakPath, DirectoryIt.Key()));
+						}
+					}
+					if (bIncludeDirectories)
+					{
+						if (Directory != PakPath)
+						{
+							DirectoriesInPak.Add(PakPath);
+						}
+					}
+				}
+				else
+				{
+					int32 SubDirIndex = PakPath.Len() > Directory.Len() ? PakPath.Find(TEXT("/"), ESearchCase::CaseSensitive, ESearchDir::FromStart, Directory.Len() + 1) : INDEX_NONE;
+					// Add files in the specified folder only.
+					if (bIncludeFiles && SubDirIndex == INDEX_NONE)
+					{
+						for (FPakDirectory::TConstIterator DirectoryIt(It.Value()); DirectoryIt; ++DirectoryIt)
+						{
+							OutFiles.Add(PakPathCombine(PakPath, DirectoryIt.Key()));
+						}
+					}
+					// Add sub-folders in the specified folder only
+					if (bIncludeDirectories && SubDirIndex >= 0)
+					{
+						DirectoriesInPak.AddUnique(PakPath.Left(SubDirIndex + 1));
+					}
+				}
+			}
 		}
+		OutFiles.Append(DirectoriesInPak);
 	}
+
+	/** Converts the path to a RelativePathFromMount and normalizes it to the expected format for Pak Directories.  Returns false if Path is not under the MountDir and hence can not be in this PakFile. */
+	bool NormalizeDirectoryQuery(const TCHAR* InPath, FString& OutRelativePathFromMount) const;
+
+	/**
+	 * Looks up the given Normalized RelativePath in the Pruned DirectoryIndex and returns the directory if found.
+	 * Validates the result if IsPakValidatePruning. Does not use a critical section; caller is responsible for calling from within an FScopedPakDirectoryIndexAccess
+	 */
+	const FPakDirectory* FindPrunedDirectoryInternal(const FString& RelativePathFromMount) const;
+
+#if ENABLE_PAKFILE_RUNTIME_PRUNING_VALIDATE
+	/* Logs an error if the two sets are not identical after removing all config-specified ignore paths */
+	void ValidateDirectorySearch(const TSet<FString>& FoundFullFiles, const TSet<FString>& PrunedFoundFiles, const TCHAR* InPath) const;
+#endif
 };
 
 /**
@@ -1471,13 +1923,19 @@ class PAKFILE_API FPakPlatformFile : public IPlatformFile
 		Paks.Append(PakFiles);
 	}
 
+	UE_DEPRECATED(4.26, "Use DirectoryExistsInPrunedPakFiles instead")
+	bool DirectoryExistsInPakFiles(const TCHAR* Directory)
+	{
+		return DirectoryExistsInPrunedPakFiles(Directory);
+	}
+
 	/**
 	 * Checks if a directory exists in one of the available pak files.
 	 *
 	 * @param Directory Directory to look for.
 	 * @return true if the directory exists, false otherwise.
 	 */
-	bool DirectoryExistsInPakFiles(const TCHAR* Directory)
+	bool DirectoryExistsInPrunedPakFiles(const TCHAR* Directory)
 	{
 		FString StandardPath = Directory;
 		FPaths::MakeStandardFilename(StandardPath);
@@ -1488,7 +1946,7 @@ class PAKFILE_API FPakPlatformFile : public IPlatformFile
 		// Check all pak files.
 		for (int32 PakIndex = 0; PakIndex < Paks.Num(); PakIndex++)
 		{
-			if (Paks[PakIndex].PakFile->DirectoryExists(*StandardPath))
+			if (Paks[PakIndex].PakFile->DirectoryExistsInPruned(*StandardPath))
 			{
 				return true;
 			}
@@ -1527,7 +1985,13 @@ class PAKFILE_API FPakPlatformFile : public IPlatformFile
 	/**
 	 * Handler for device delegate to prompt us to load a new pak.	 
 	 */
-	bool HandleMountPakDelegate(const FString& PakFilePath, int32 PakOrder, IPlatformFile::FDirectoryVisitor* Visitor);
+	IPakFile* HandleMountPakDelegate(const FString& PakFilePath, int32 PakOrder);
+
+	/**
+	 * Handler for device delegate to prompt us to load a new pak.
+	 */
+	UE_DEPRECATED(4.26, "Use HandleMountPakDelegate instead")
+	bool HandleOnMountPakDelegate(const FString& PakFilePath, int32 PakOrder, IPlatformFile::FDirectoryVisitor* Visitor);
 
 	/**
 	 * Handler for device delegate to prompt us to unload a pak.
@@ -1691,7 +2155,7 @@ public:
 	/**
 	 * Make unique in memory pak files from a list of named files
 	 */
-	virtual void MakeUniquePakFilesForTheseFiles(TArray<TArray<FString>> InFiles);
+	virtual void MakeUniquePakFilesForTheseFiles(const TArray<TArray<FString>>& InFiles);
 
 
 	/**
@@ -1719,7 +2183,7 @@ public:
 				return false;
 			}
 
-			FPakFile::EFindResult FindResult = Paks[PakIndex].PakFile->Find(*StandardFilename, OutEntry);
+			FPakFile::EFindResult FindResult = Paks[PakIndex].PakFile->Find(StandardFilename, OutEntry);
 			if (FindResult == FPakFile::EFindResult::Found )
 			{
 				if (OutPakFile != NULL)
@@ -1937,25 +2401,35 @@ public:
 		FPakFile* PakFile = NULL;
 		if (FindFileInPakFiles(Filename, &PakFile, &FileEntry))
 		{
-			if (!PakFile->HasFilenames())
-			{
-				//FPlatformMisc::LowLevelOutputDebugString(*(FString("GetFilenameOfDisk() used when bFilenamesRemoved == true: ") + Filename));
-			}
-			//checkf(PakFile->HasFilenames(), TEXT("GetFilenameOnDisk() can only be used before FPakPlatformFile::UnloadPakEntryFilenames() is called."));
-
 			const FString Path(FPaths::GetPath(Filename));
-			const FPakDirectory* PakDirectory = PakFile->FindDirectory(*Path);
+			FPakFile::FScopedPakDirectoryIndexAccess ScopeAccess(*PakFile);
+
+			const FPakDirectory* PakDirectory = PakFile->FindPrunedDirectory(*Path);
 			if (PakDirectory != nullptr)
 			{
 				for (FPakDirectory::TConstIterator DirectoryIt(*PakDirectory); DirectoryIt; ++DirectoryIt)
 				{
-					if (PakFile->Data->Files[DirectoryIt.Value()].Offset == FileEntry.Offset)
+					FPakEntry PakEntry;
+					if (PakFile->GetPakEntry(DirectoryIt.Value(), &PakEntry) != FPakFile::EFindResult::NotFound && PakEntry.Offset == FileEntry.Offset)
 					{
 						const FString& RealFilename = DirectoryIt.Key();
 						return Path / RealFilename;
 					}
 				}
 			}
+
+#if ENABLE_PAKFILE_RUNTIME_PRUNING_VALIDATE
+			// The File exists in the Pak but has been pruned from its DirectoryIndex; log an error if we are validating pruning and return the original Filename.
+			if (PakFile->ShouldValidatePrunedDirectory())
+			{
+				TSet<FString> FullFoundFiles;
+				TSet<FString> PrunedFoundFiles;
+				FullFoundFiles.Add(Filename);
+				PakFile->ValidateDirectorySearch(FullFoundFiles, PrunedFoundFiles, Filename);
+			}
+#endif
+
+			return Filename;
 		}
 
 		// Fall back to lower level.
@@ -1985,7 +2459,7 @@ public:
 	virtual bool DirectoryExists(const TCHAR* Directory) override
 	{
 		// Check pak files first.
-		if (DirectoryExistsInPakFiles(Directory))
+		if (DirectoryExistsInPrunedPakFiles(Directory))
 		{
 			return true;
 		}
@@ -2003,9 +2477,9 @@ public:
 	virtual bool DeleteDirectory(const TCHAR* Directory) override
 	{
 		// Even if the same directory exists outside of pak files it will never
-		// get truely deleted from pak and will still be reported by Iterate functions.
+		// get truly deleted from pak and will still be reported by Iterate functions.
 		// Fail in cases like this.
-		if (DirectoryExistsInPakFiles(Directory))
+		if (DirectoryExistsInPrunedPakFiles(Directory))
 		{
 			return false;
 		}
@@ -2031,7 +2505,7 @@ public:
 		}
 
 		// Then check pak directories
-		if (DirectoryExistsInPakFiles(FilenameOrDirectory))
+		if (DirectoryExistsInPrunedPakFiles(FilenameOrDirectory))
 		{
 			return FFileStatData(
 				PakFile->GetTimestamp(),
@@ -2119,7 +2593,7 @@ public:
 			const bool bIncludeFolders = true;
 			TSet<FString> FilesVisitedInThisPak;
 
-			PakFile.FindFilesAtPath(FilesVisitedInThisPak, *StandardDirectory, bIncludeFiles, bIncludeFolders);
+			PakFile.FindPrunedFilesAtPath(FilesVisitedInThisPak, *StandardDirectory, bIncludeFiles, bIncludeFolders);
 			for (TSet<FString>::TConstIterator SetIt(FilesVisitedInThisPak); SetIt && Result; ++SetIt)
 			{
 				const FString& Filename = *SetIt;
@@ -2231,7 +2705,7 @@ public:
 			const bool bIncludeFolders = true;
 			TSet<FString> FilesVisitedInThisPak;
 
-			PakFile.FindFilesAtPath(FilesVisitedInThisPak, *StandardDirectory, bIncludeFiles, bIncludeFolders);
+			PakFile.FindPrunedFilesAtPath(FilesVisitedInThisPak, *StandardDirectory, bIncludeFiles, bIncludeFolders);
 			for (TSet<FString>::TConstIterator SetIt(FilesVisitedInThisPak); SetIt && Result; ++SetIt)
 			{
 				const FString& Filename = *SetIt;
@@ -2338,7 +2812,7 @@ public:
 			for (int32 PakIndex = 0; PakIndex < Paks.Num(); PakIndex++)
 			{
 				FPakFile& PakFile = *Paks[PakIndex].PakFile;
-				PakFile.FindFilesAtPath(FilesInPak, *StandardDirectory, bIncludeFiles, bIncludeFolders, bRecursive);
+				PakFile.FindPrunedFilesAtPath(FilesInPak, *StandardDirectory, bIncludeFiles, bIncludeFolders, bRecursive);
 			}
 			
 			for (const FString& Filename : FilesInPak)
@@ -2362,11 +2836,11 @@ public:
 			}
 		}
 	}
-	
+
 	virtual bool DeleteDirectoryRecursively(const TCHAR* Directory) override
 	{
 		// Can't delete directories existing in pak files. See DeleteDirectory(..) for more info.
-		if (DirectoryExistsInPakFiles(Directory))
+		if (DirectoryExistsInPrunedPakFiles(Directory))
 		{
 			return false;
 		}
@@ -2458,12 +2932,25 @@ public:
 	// Broadacast a master signature table failure through any registered delegates in a thread safe way
 	static void BroadcastPakMasterSignatureTableCheckFailure(const FString& InFilename);
 
-	// Get a list of which files live in a given chunk
-	void GetFilenamesInChunk(const FString& InPakFilename, const TArray<int32>& InChunkIDs, TArray<FString>& OutFileList);
-	void GetFilenamesInPakFile(const FString& InPakFilename, TArray<FString>& OutFileList) ;
+	// Access static delegate for setting PakIndex settings.
+	static FPakSetIndexSettings& GetPakSetIndexSettingsDelegate();
+
+	UE_DEPRECATED(4.26, "Use GetPrunedFilenamesInChunk")
+	void GetFilenamesInChunk(const FString& InPakFilename, const TArray<int32>& InChunkIDs, TArray<FString>& OutFileList)
+	{
+		GetPrunedFilenamesInChunk(InPakFilename, InChunkIDs, OutFileList);
+	}
+
+	/* Get a list of RelativePathFromMount for every file in the given Pak that lives in any of the given chunks.  Only searches the Pruned DirectoryIndex */
+	void GetPrunedFilenamesInChunk(const FString& InPakFilename, const TArray<int32>& InChunkIDs, TArray<FString>& OutFileList);
+	UE_DEPRECATED(4.26, "Use GetPrunedFilenamesInPakFile")
+	void GetFilenamesInPakFile(const FString& InPakFilename, TArray<FString>& OutFileList)
+	{
+		GetPrunedFilenamesInPakFile(InPakFilename, OutFileList);
+	}
 
-	void UnloadPakEntryFilenames(TArray<FString>* DirectoryRootsToKeep = nullptr);
-	void ShrinkPakEntriesMemoryUsage();
+	/** Gets a list of FullPaths (includes Mount directory) for every File in the given Pak's Pruned DirectoryIndex */
+	void GetPrunedFilenamesInPakFile(const FString& InPakFilename, TArray<FString>& OutFileList);
 
 	// BEGIN Console commands
 #if !UE_BUILD_SHIPPING
diff --git a/Engine/Source/Runtime/Projects/Private/PluginManager.cpp b/Engine/Source/Runtime/Projects/Private/PluginManager.cpp
--- a/Engine/Source/Runtime/Projects/Private/PluginManager.cpp
+++ b/Engine/Source/Runtime/Projects/Private/PluginManager.cpp
@@ -828,14 +828,14 @@ bool FPluginManager::ConfigureEnabledPlugins()
 					PlatformFile.IterateDirectoryRecursively(*(ContentDir / TEXT("Paks") / FPlatformProperties::PlatformName()), PakVisitor);
 					for (const FString& PakPath : FoundPaks)
 					{
-						if (FCoreDelegates::OnMountPak.IsBound())
+						if (FCoreDelegates::MountPak.IsBound())
 						{
-							FCoreDelegates::OnMountPak.Execute(PakPath, 0, nullptr);
+							FCoreDelegates::MountPak.Execute(PakPath, 0);
 							PluginsWithPakFile.AddUnique(Plugin);
 						}
 						else
 						{
-							UE_LOG(LogPluginManager, Warning, TEXT("PAK file (%s) could not be mounted because OnMountPak is not bound"), *PakPath)
+							UE_LOG(LogPluginManager, Warning, TEXT("PAK file (%s) could not be mounted because MountPak is not bound"), *PakPath)
 						}
 					}
 				}
