Fix localization not detecting the existance of some localization files and falling back to english for those files.   FPakFile::NormalizeDirectoryQuery was incorrectly returning false when queried for whether the MountPoint exists.

diff --git a/Engine/Source/Runtime/PakFile/Private/IPlatformFilePak.cpp b/Engine/Source/Runtime/PakFile/Private/IPlatformFilePak.cpp
--- a/Engine/Source/Runtime/PakFile/Private/IPlatformFilePak.cpp
+++ b/Engine/Source/Runtime/PakFile/Private/IPlatformFilePak.cpp
@@ -5734,6 +5734,7 @@ void FPakFile::AddEntryToIndex(const FString& Filename, const FPakEntryLocation&
 	{
 		check(IsPathInDirectoryFormat(MountPoint));
 		check(Filename.StartsWith(MountPoint));
+		check(Filename.Len() > MountPoint.Len());
 		RelativePathFromMount = Filename.Mid(MountPoint.Len());
 	}
 
@@ -6030,14 +6031,7 @@ bool FPakFile::NormalizeDirectoryQuery(const TCHAR* InPath, FString& OutRelative
 {
 	OutRelativePathFromMount = InPath;
 	MakeDirectoryFromPath(OutRelativePathFromMount);
-
-	// Check the specified path is under the mount point of this pak file.
-	if (!OutRelativePathFromMount.StartsWith(MountPoint))
-	{
-		return false;
-	}
-	OutRelativePathFromMount = OutRelativePathFromMount.Mid(MountPoint.Len());
-	return true;
+	return GetRelativePathFromMountInline(OutRelativePathFromMount, MountPoint);
 }
 
 const FPakDirectory* FPakFile::FindPrunedDirectoryInternal(const FString& RelativePathFromMount) const
diff --git a/Engine/Source/Runtime/PakFile/Public/IPlatformFilePak.h b/Engine/Source/Runtime/PakFile/Public/IPlatformFilePak.h
--- a/Engine/Source/Runtime/PakFile/Public/IPlatformFilePak.h
+++ b/Engine/Source/Runtime/PakFile/Public/IPlatformFilePak.h
@@ -1485,6 +1485,26 @@ public:
 		}
 	}
 
+	/**
+	 * Helper function to return Child's relative path from the mount point.  Returns false if Child is not equal to MountPoint and is not a child path under MountPoint, else returns true.
+	 * Edits Child only if returning true, setting it to the relative path.
+	 * If child equals MountPoint, returns true and sets Child to the relative path to the MountPoint, which is "/"
+	 */
+	static bool GetRelativePathFromMountInline(FString& Child, const FString& MountPoint)
+	{
+		if (!Child.StartsWith(MountPoint))
+		{
+			return false;
+		}
+		Child = Child.Mid(MountPoint.Len());
+		if (Child.IsEmpty())
+		{
+			// Child equals the MountPoint
+			Child = TEXT("/");
+		}
+		return true;
+	}
+
 	/* Returns the global,const flag for whether the current process is allowing PakFiles to keep their entire DirectoryIndex (if it exists in the PakFile on disk) rather than pruning it */
 	static bool IsPakKeepFullDirectory();
 
