Revert the invalidation of revision PakFile_Version_FrozenIndex; that change was incorrectly merged to main.
#rnx

diff --git a/Engine/Source/Runtime/PakFile/Public/IPlatformFilePak.h b/Engine/Source/Runtime/PakFile/Public/IPlatformFilePak.h
--- a/Engine/Source/Runtime/PakFile/Public/IPlatformFilePak.h
+++ b/Engine/Source/Runtime/PakFile/Public/IPlatformFilePak.h
@@ -210,11 +210,6 @@ struct FPakInfo
 			{
 				EncryptionKeyGuid.Invalidate();
 			}
-
-			if (Version == PakFile_Version_FrozenIndex)
-			{
-				UE_LOG(LogPakFile, Fatal, TEXT("FPakInfo version PakFile_Version_FrozenIndex is no longer supported; regenerate Pak files"));
-			}
 		}
 
 		if (Version >= PakFile_Version_FrozenIndex && Version < PakFile_Version_PathHashIndex)
