Fix a deadlock when loading the ShooterGame sample from a packaged build.

FChunkCacheWorker::ChunkRequestAvailable could have many threads waiting on it and would only wake up a single thread each time a FChunkBuffer is processed. There would be no guarantee that the correct thread would be woken up or that each thread would be woken a single time, meaning that sooner or later we would end up with the FChunkCacheWorker having no more work to process but a number of threads still idle, waiting on the event to wake them, which would never happen.

Switching to a manually reset event will cause all waiting threads to be woken each time the event is triggered and will make sure that we don't deadlock.
This issue showed up now due to changes else where in the code which resulted in more threads waiting on FChunkCacheWorker than we previously had.

diff --git a/Engine/Source/Runtime/PakFile/Private/SignedArchiveReader.cpp b/Engine/Source/Runtime/PakFile/Private/SignedArchiveReader.cpp
--- a/Engine/Source/Runtime/PakFile/Private/SignedArchiveReader.cpp
+++ b/Engine/Source/Runtime/PakFile/Private/SignedArchiveReader.cpp
@@ -44,7 +44,7 @@ FChunkCacheWorker::FChunkCacheWorker(FArchive* InReader, const TCHAR* Filename)
 		if (bEnableMultithreading)
 		{
 			QueuedRequestsEvent = FPlatformProcess::GetSynchEventFromPool();
-			ChunkRequestAvailable = FPlatformProcess::GetSynchEventFromPool();
+			ChunkRequestAvailable = FPlatformProcess::GetSynchEventFromPool(true);
 			Thread = FRunnableThread::Create(this, TEXT("FChunkCacheWorker"), 0, TPri_BelowNormal);
 		}
 	}
@@ -211,6 +211,7 @@ int32 FChunkCacheWorker::ProcessQueue()
 				if (ChunkRequestAvailable != nullptr)
 				{
 					ChunkRequestAvailable->Trigger();
+					ChunkRequestAvailable->Reset();
 				}
 			}
 		}
