Fix incorrect use of static constexpr without inline.
#rnx

diff --git a/Engine/Source/Runtime/PakFile/Public/IPlatformFilePak.h b/Engine/Source/Runtime/PakFile/Public/IPlatformFilePak.h
--- a/Engine/Source/Runtime/PakFile/Public/IPlatformFilePak.h
+++ b/Engine/Source/Runtime/PakFile/Public/IPlatformFilePak.h
@@ -533,8 +533,8 @@ public:
 	 * 0x80000000: Invalid
 	 * 0x80000001 - 0xffffffff: FileIndex from MaxIndex to 0
 	*/
-	static CONSTEXPR int32 Invalid = MIN_int32;
-	static CONSTEXPR int32 MaxIndex = MAX_int32 - 1;
+	static const int32 Invalid = MIN_int32;
+	static const int32 MaxIndex = MAX_int32 - 1;
 
 	FPakEntryLocation() : Index(Invalid)
 	{
