#rnx

diff --git a/Engine/Source/Runtime/PakFile/Public/IPlatformFilePak.h b/Engine/Source/Runtime/PakFile/Public/IPlatformFilePak.h
--- a/Engine/Source/Runtime/PakFile/Public/IPlatformFilePak.h
+++ b/Engine/Source/Runtime/PakFile/Public/IPlatformFilePak.h
@@ -7,6 +7,7 @@
 #include "Stats/Stats.h"
 #include "Misc/Paths.h"
 #include "Misc/ScopeLock.h"
+#include "Misc/ScopeRWLock.h"
 #include "Templates/UniquePtr.h"
 #include "Math/BigInt.h"
 #include "Misc/AES.h"
@@ -32,10 +33,6 @@ DECLARE_DWORD_ACCUMULATOR_STAT_EXTERN(TEXT("Num open pak file handles"), STAT_Pa
 #endif
 #define ENABLE_PAKFILE_RUNTIME_PRUNING_VALIDATE ENABLE_PAKFILE_RUNTIME_PRUNING && !UE_BUILD_SHIPPING
 
-#if ENABLE_PAKFILE_RUNTIME_PRUNING
-#include "Misc/ScopeRWLock.h"
-#endif
-
 // Define the type of a chunk hash. Currently selectable between SHA1 and CRC32.
 #define PAKHASH_USE_CRC	1
 #if PAKHASH_USE_CRC
