Fix a potential deadlock when loading from packfiles

This was originally fixed in CL 11458603 but it did not fix the issue entirely and in any case I was not happy with the manual reset event approach.
The problem with the manual reset event was that there is no actual safe point at which we can reset the event without risking some requesting thread missing it and waiting forever, it is just a lot rarer than the original thread lock.

- Each FChunkRequest now can be given an event that is triggered, so the correct thread is always woken up if waiting.
- Since the signals are 1-1 we no longer risk having threads never being woken.
- We only wake the thread that has a request to process so don't risk flooding the cpu.
- Removed FChunkCacheWorker::SetupDecryptionKey as it was not implemented/used.
- Removed FChunkCacheWorker::ChunkRequestAvailable and its associated methods ::WaitForNextChunk/::FlushRemainingChunkCompletionEvents.

In testing this version runs slightly faster on loading tests (0.2seconds on average), although just removing the events entirely and running a busy loop with a yield ran slightly faster on machines with less cores (4 cores with hyperthreading or less) albiet with a much higher variation and so potentially could end up a lot worse.

diff --git a/Engine/Source/Runtime/PakFile/Private/SignedArchiveReader.cpp b/Engine/Source/Runtime/PakFile/Private/SignedArchiveReader.cpp
--- a/Engine/Source/Runtime/PakFile/Private/SignedArchiveReader.cpp
+++ b/Engine/Source/Runtime/PakFile/Private/SignedArchiveReader.cpp
@@ -34,7 +34,6 @@ FChunkCacheWorker::FChunkCacheWorker(FArchive* InReader, const TCHAR* Filename)
 	: Thread(nullptr)
 	, Reader(InReader)
 	, QueuedRequestsEvent(nullptr)
-	, ChunkRequestAvailable(nullptr)
 {
 	Signatures = FPakPlatformFile::GetPakSignatureFile(Filename);
 
@@ -44,7 +43,6 @@ FChunkCacheWorker::FChunkCacheWorker(FArchive* InReader, const TCHAR* Filename)
 		if (bEnableMultithreading)
 		{
 			QueuedRequestsEvent = FPlatformProcess::GetSynchEventFromPool();
-			ChunkRequestAvailable = FPlatformProcess::GetSynchEventFromPool(true);
 			Thread = FRunnableThread::Create(this, TEXT("FChunkCacheWorker"), 0, TPri_BelowNormal);
 		}
 	}
@@ -59,11 +57,6 @@ FChunkCacheWorker::~FChunkCacheWorker()
 		FPlatformProcess::ReturnSynchEventToPool(QueuedRequestsEvent);
 		QueuedRequestsEvent = nullptr;
 	}
-	if (ChunkRequestAvailable != nullptr)
-	{ 
-		FPlatformProcess::ReturnSynchEventToPool(ChunkRequestAvailable);
-		ChunkRequestAvailable = nullptr;
-	}
 }
 
 bool FChunkCacheWorker::Init()
@@ -208,10 +201,10 @@ int32 FChunkCacheWorker::ProcessQueue()
 				// Let the other thread know the chunk is ready to read.
 				Request.RefCount.Decrement();
 				Request.IsTrusted.Increment();
-				if (ChunkRequestAvailable != nullptr)
+
+				if (Request.Event != nullptr)
 				{
-					ChunkRequestAvailable->Trigger();
-					ChunkRequestAvailable->Reset();
+					Request.Event->Trigger();
 				}
 			}
 		}
@@ -259,7 +252,7 @@ bool FChunkCacheWorker::CheckSignature(const FChunkRequest& ChunkInfo)
 	return bChunkHashesMatch;
 }
 
-FChunkRequest& FChunkCacheWorker::RequestChunk(int32 ChunkIndex, int64 StartOffset, int64 ChunkSize)
+FChunkRequest& FChunkCacheWorker::RequestChunk(int32 ChunkIndex, int64 StartOffset, int64 ChunkSize, FEvent* Event)
 {
 	FChunkRequest* NewChunk = FreeChunkRequests.Pop();
 	if (NewChunk == NULL)
@@ -273,6 +266,7 @@ FChunkRequest& FChunkCacheWorker::RequestChunk(int32 ChunkIndex, int64 StartOffs
 	NewChunk->IsTrusted.Set(0);
 	// At this point both worker and the archive use this chunk so increase ref count
 	NewChunk->RefCount.Set(2);
+	NewChunk->Event = Event;
 
 	QueueLock.Lock();
 	RequestQueue.Add(NewChunk);
@@ -286,22 +280,6 @@ FChunkRequest& FChunkCacheWorker::RequestChunk(int32 ChunkIndex, int64 StartOffs
 	return *NewChunk;
 }
 
-void FChunkCacheWorker::WaitForNextChunk()
-{
-	if (ChunkRequestAvailable != nullptr)
-	{
-		ChunkRequestAvailable->Wait();
-	}
-}
-
-void FChunkCacheWorker::FlushRemainingChunkCompletionEvents()
-{
-	if (ChunkRequestAvailable != nullptr)
-	{
-		ChunkRequestAvailable->Reset();
-	}
-}
-
 bool FChunkCacheWorker::IsValid() const
 {
 	return Signatures.IsValid() && (Signatures->ChunkHashes.Num() > 0);
@@ -358,7 +336,7 @@ int64 FSignedArchiveReader::CalculateChunkSize(int64 ChunkIndex) const
 	}
 }
 
-int64 FSignedArchiveReader::PrecacheChunks(TArray<FSignedArchiveReader::FReadInfo>& Chunks, int64 Length)
+int64 FSignedArchiveReader::PrecacheChunks(TArray<FSignedArchiveReader::FReadInfo>& Chunks, int64 Length, FEvent* Event)
 {
 	SCOPE_SECONDS_ACCUMULATOR(STAT_SignedArchiveReader_PreCacheChunks);
 
@@ -402,7 +380,7 @@ int64 FSignedArchiveReader::PrecacheChunks(TArray<FSignedArchiveReader::FReadInf
 		else
 		{
 			const int64 ChunkSize = CalculateChunkSize(ChunkIndex);	
-			ChunkInfo.Request = &SignatureChecker->RequestChunk(ChunkIndex, ChunkStartOffset, ChunkSize);
+			ChunkInfo.Request = &SignatureChecker->RequestChunk(ChunkIndex, ChunkStartOffset, ChunkSize, Event);
 			INC_DWORD_STAT(STAT_SignedArchiveReader_NumChunkRequests);
 			ChunkInfo.PreCachedChunk = NULL;
 		}
@@ -422,9 +400,11 @@ void FSignedArchiveReader::Serialize(void* Data, int64 Length)
 	SCOPE_SECONDS_ACCUMULATOR(STAT_SignedArchiveReader_Serialize);
 	INC_DWORD_STAT(STAT_SignedArchiveReader_NumSerializes);
 
+	FEvent* ChunkReadEvent = SignatureChecker->IsMultithreaded() ? FPlatformProcess::GetSynchEventFromPool() : nullptr;
+	
 	// First make sure the chunks we're going to read are actually cached.
 	TArray<FReadInfo> QueuedChunks;
-	int64 ChunksToRead = PrecacheChunks(QueuedChunks, Length);
+	int64 ChunksToRead = PrecacheChunks(QueuedChunks, Length, ChunkReadEvent);
 	int64 FirstPrecacheChunkIndex = ChunksToRead;
 
 	// If we aren't multithreaded then flush the signature checking now so there will be some data ready
@@ -485,7 +465,7 @@ void FSignedArchiveReader::Serialize(void* Data, int64 Length)
 			{
 				if (SignatureChecker->IsMultithreaded())
 				{
-					SignatureChecker->WaitForNextChunk();
+					ChunkReadEvent->Wait();
 				}
 				else
 				{
@@ -497,10 +477,7 @@ void FSignedArchiveReader::Serialize(void* Data, int64 Length)
 		while (ChunksToRead > 0);
 	}
 
-	// Need to flush out any remaining request events here. Each time the loop above wakes up on the event,
-	// it will process EVERY chunk available. It may process 2 available chunks from that one trigger then complete,
-	// leaving a single event trigger outstanding which will break the next call into this reader
-	SignatureChecker->FlushRemainingChunkCompletionEvents();
+	FPlatformProcess::ReturnSynchEventToPool(ChunkReadEvent);
 
 	PakOffset += Length;
 
@@ -513,4 +490,4 @@ void FSignedArchiveReader::Serialize(void* Data, int64 Length)
 			SignatureChecker->ReleaseChunk(*CachedChunk.Request);
 		}
 	}
-}
\ No newline at end of file
+}
diff --git a/Engine/Source/Runtime/PakFile/Private/SignedArchiveReader.h b/Engine/Source/Runtime/PakFile/Private/SignedArchiveReader.h
--- a/Engine/Source/Runtime/PakFile/Private/SignedArchiveReader.h
+++ b/Engine/Source/Runtime/PakFile/Private/SignedArchiveReader.h
@@ -59,6 +59,8 @@ struct FChunkRequest
 	FThreadSafeCounter IsTrusted;
 	/** Reference count */
 	FThreadSafeCounter RefCount;
+	/** Optional pointer to an FEvent that can be used to signal that the request is ready for use */
+	FEvent* Event;
 
 	/**
 	 * Constructor
@@ -67,9 +69,10 @@ struct FChunkRequest
 		: Index(INDEX_NONE)
 		, Offset(0)
 		, Size(0)
-		, Buffer(NULL)
+		, Buffer(nullptr)
 		, IsTrusted(0)
 		, RefCount(0)
+		, Event(nullptr)
 	{}
 
 	/**
@@ -122,8 +125,6 @@ class FChunkCacheWorker : public FRunnable
 	FThreadSafeCounter PendingQueueCounter;
 	/** Event used to signal there's work to be done */
 	FEvent* QueuedRequestsEvent;
-	/** Event used to signal there's completed work to be processed */
-	FEvent* ChunkRequestAvailable;
 	/** List of active chunk requests */
 	TArray<FChunkRequest*> ActiveRequests;
 	/** Stops this thread */
@@ -154,10 +155,6 @@ class FChunkCacheWorker : public FRunnable
 	 */
 	void ReleaseBuffer(int32 ChunkIndex);
 	/**
-	* Initializes the public key
-	*/
-	void SetupDecryptionKey();
-	/**
 	* Is this chunk cache worker running in a thread?
 	*/
 	FORCEINLINE bool IsMultithreaded() const 
@@ -165,17 +162,6 @@ class FChunkCacheWorker : public FRunnable
 		return Thread != nullptr;
 	}
 
-	/**
-	* Block until there is a new chunk to process on the main thread
-	*/
-	void WaitForNextChunk();
-
-	/**
-	* Reset any outstanding chunk completion event triggers that may still be 
-	* left over when the main thread has finished copying out all the data it needs
-	*/
-	void FlushRemainingChunkCompletionEvents();
-
 public:
 
 	FChunkCacheWorker(FArchive* InReader, const TCHAR* Filename);
@@ -193,9 +179,10 @@ public:
 	 * @param ChunkIndex Index of a chunk to load
 	 * @param StartOffset Offset to the beginning of the chunk
 	 * @param ChunkSize Chunk size
+	 * @param Event Optional FEvent that will signal when the request is ready, nullptr is valid if the calling code does not want this signal
 	 * @return Handle to the request.
 	 */
-	FChunkRequest& RequestChunk(int32 ChunkIndex, int64 StartOffset, int64 ChunkSize);
+	FChunkRequest& RequestChunk(int32 ChunkIndex, int64 StartOffset, int64 ChunkSize, FEvent* Event);
 	/**
 	 * Releases the requested chunk buffer
 	 */
@@ -285,9 +272,12 @@ class FSignedArchiveReader : public FArchive
 
 	/** 
 	 * Queues chunks on the worker thread 
+	 * @param Chunks This array will contain info about each chunk created by the call.
+	 * @param Length The length of data to precache
+	 * @param Event Optional FEvent that will signal as each request becomes ready, nullptr is valid if the calling code does not want this signal
 	 * @return Number of chunks in the output array which are actually required for the requested length. The rest are precache chunks 
 	 */
-	int64 PrecacheChunks(TArray<FReadInfo>& Chunks, int64 Length);
+	int64 PrecacheChunks(TArray<FReadInfo>& Chunks, int64 Length, FEvent* Event);
 
 public:
 
