No longer load verify data if verify code is turned off.
#testing iOS Client

diff --git a/Engine/Source/Runtime/PakFile/Private/IPlatformFilePak.cpp b/Engine/Source/Runtime/PakFile/Private/IPlatformFilePak.cpp
--- a/Engine/Source/Runtime/PakFile/Private/IPlatformFilePak.cpp
+++ b/Engine/Source/Runtime/PakFile/Private/IPlatformFilePak.cpp
@@ -5183,7 +5183,7 @@ bool FPakFile::LoadIndexInternal(FArchive* Reader)
 		}
 		bHasPathHashIndex = true;
 	}
-
+	
 	if (!bReadFullDirectoryIndex)
 	{
 		check(bWillUsePathHashIndex); // Need to confirm that we have read the PathHashIndex bytes
@@ -5681,8 +5681,12 @@ bool FPakFile::IsPakKeepFullDirectory()
 
 bool FPakFile::IsPakValidatePruning()
 {
+#if ENABLE_PAKFILE_RUNTIME_PRUNING_VALIDATE
 	FIndexSettings& IndexLoadParams = GetIndexSettings();
 	return IndexLoadParams.bValidatePruning;
+#else
+	return false;
+#endif
 }
 
 bool FPakFile::IsPakDelayPruning()
