Legacy pak file versions encode pak entries in FPakFile::LoadLegacyIndex. That function repeatedly serializes bytes from the pakfile on disk into a reused FPakEntry. Serialize expects the FPakEntry to be in an empty constructed state, and it was not empty due to serializing previous PakEntries into it.
Added a reset function and call it in this one place where FPakEntry is reused.
#rnx

diff --git a/Engine/Source/Runtime/PakFile/Private/IPlatformFilePak.cpp b/Engine/Source/Runtime/PakFile/Private/IPlatformFilePak.cpp
--- a/Engine/Source/Runtime/PakFile/Private/IPlatformFilePak.cpp
+++ b/Engine/Source/Runtime/PakFile/Private/IPlatformFilePak.cpp
@@ -5318,6 +5318,7 @@ bool FPakFile::LoadLegacyIndex(FArchive* Reader)
 	auto ReadNextEntry = [&PakEntryPair, &IndexReader, this]() -> FPakEntryPair&
 	{
 		IndexReader << PakEntryPair.Filename;
+		PakEntryPair.Info.Reset();
 		PakEntryPair.Info.Serialize(IndexReader, this->Info.Version);
 		return PakEntryPair;
 	};
diff --git a/Engine/Source/Runtime/PakFile/Public/IPlatformFilePak.h b/Engine/Source/Runtime/PakFile/Public/IPlatformFilePak.h
--- a/Engine/Source/Runtime/PakFile/Public/IPlatformFilePak.h
+++ b/Engine/Source/Runtime/PakFile/Public/IPlatformFilePak.h
@@ -352,15 +352,8 @@ struct FPakEntry
 	 * Constructor.
 	 */
 	FPakEntry()
-		: Offset(-1)
-		, Size(0)
-		, UncompressedSize(0)
-		, CompressionBlockSize(0)
-		, CompressionMethodIndex(0)
-		, Flags(Flag_None)
-		, Verified(false)
 	{
-		FMemory::Memset(Hash, 0, sizeof(Hash));
+		Reset();
 	}
 
 	/**
@@ -427,6 +420,19 @@ struct FPakEntry
 			CompressionBlocks == B.CompressionBlocks;
 	}
 
+	void Reset()
+	{
+		Offset = -1;
+		Size = 0;
+		UncompressedSize = 0;
+		FMemory::Memset(Hash, 0, sizeof(Hash));
+		CompressionBlocks.Reset();
+		CompressionBlockSize = 0;
+		CompressionMethodIndex = 0;
+		Flags = Flag_None;
+		Verified = false;
+	}
+
 	/**
 	 * Serializes FPakEntry struct.
 	 *
