[RN] Minor, PakFiles

diff --git a/Engine/Source/Runtime/PakFile/Private/IPlatformFilePak.cpp b/Engine/Source/Runtime/PakFile/Private/IPlatformFilePak.cpp
--- a/Engine/Source/Runtime/PakFile/Private/IPlatformFilePak.cpp
+++ b/Engine/Source/Runtime/PakFile/Private/IPlatformFilePak.cpp
@@ -5738,9 +5738,10 @@ void FPakFile::AddEntryToIndex(const FString& Filename, const FPakEntryLocation&
 	else
 	{
 		check(IsPathInDirectoryFormat(MountPoint));
-		check(Filename.StartsWith(MountPoint));
+		RelativePathFromMount = Filename;
 		check(Filename.Len() > MountPoint.Len());
-		RelativePathFromMount = Filename.Mid(MountPoint.Len());
+		bool bSucceeded = GetRelativePathFromMountInline(RelativePathFromMount, MountPoint);
+		check(bSucceeded);
 	}
 
 	if (DirectoryIndex)
@@ -6377,11 +6378,11 @@ FArchive* FPakFile::GetSharedReader(IPlatformFile* LowerLevel)
 
 const FPakEntryLocation* FPakFile::FindLocationFromIndex(const FString& FullPath, const FString& MountPoint, const FPathHashIndex& PathHashIndex, uint64 PathHashSeed)
 {
-	if (!FullPath.StartsWith(MountPoint))
+	const TCHAR* RelativePathFromMount = GetRelativeFilePathFromMountPointer(FullPath, MountPoint);
+	if (!RelativePathFromMount)
 	{
 		return nullptr;
 	}
-	const TCHAR* RelativePathFromMount = (*FullPath) + MountPoint.Len();
 	uint64 PathHash = HashPath(RelativePathFromMount, PathHashSeed);
 	return PathHashIndex.Find(PathHash);
 }
diff --git a/Engine/Source/Runtime/PakFile/Public/IPlatformFilePak.h b/Engine/Source/Runtime/PakFile/Public/IPlatformFilePak.h
--- a/Engine/Source/Runtime/PakFile/Public/IPlatformFilePak.h
+++ b/Engine/Source/Runtime/PakFile/Public/IPlatformFilePak.h
@@ -1503,6 +1503,26 @@ public:
 		return true;
 	}
 
+	/**
+	 * Helper function to return Filename's relative path from the mount point. Returns null if Child is not equal to MountPoint and is not a child path under MountPoint, else returns
+	 * pointer to the offset within Child after the MountPoint.
+	 * If child equals MountPoint, returns null; The MountPoint itself is not a valid Filename, since Filenames must have non-zero length and are added on to the MountPoint.
+	 */
+	static const TCHAR* GetRelativeFilePathFromMountPointer(const FString& Child, const FString& MountPoint)
+	{
+		if (!Child.StartsWith(MountPoint))
+		{
+			return nullptr;
+		}
+		const TCHAR* RelativePathFromMount = (*Child) + MountPoint.Len();
+		if (RelativePathFromMount[0] == TEXT('\0'))
+		{
+			// Child is equal to MountPoint, invalid
+			return nullptr;
+		}
+		return RelativePathFromMount;
+	}
+
 	/* Returns the global,const flag for whether the current process is allowing PakFiles to keep their entire DirectoryIndex (if it exists in the PakFile on disk) rather than pruning it */
 	static bool IsPakKeepFullDirectory();
 
@@ -2553,47 +2573,87 @@ public:
 	/**
 	 * Helper class to filter out files which have already been visited in one of the pak files.
 	 */
-	class FPakVisitor : public IPlatformFile::FDirectoryVisitor
+	class FPreventDuplicatesVisitorBase
+	{
+	public:
+		/** Visited files. */
+		TSet<FString>& VisitedFiles;
+		FString NormalizedFilename;
+
+		FPreventDuplicatesVisitorBase(TSet<FString>& InVisitedFiles)
+			: VisitedFiles(InVisitedFiles)
+		{
+		}
+
+		bool CheckDuplicate(const TCHAR* FilenameOrDirectory)
+		{
+			NormalizedFilename.Reset();
+			NormalizedFilename.AppendChars(FilenameOrDirectory, TCString<TCHAR>::Strlen(FilenameOrDirectory));
+			FPaths::MakeStandardFilename(NormalizedFilename);
+			if (VisitedFiles.Contains(NormalizedFilename))
+			{
+				return true;
+			}
+			VisitedFiles.Add(NormalizedFilename);
+			return false;
+		}
+	};
+
+	class FPreventDuplicatesVisitor : public FPreventDuplicatesVisitorBase, public IPlatformFile::FDirectoryVisitor
 	{
 	public:
 		/** Wrapped visitor. */
-		FDirectoryVisitor&	Visitor;
-		/** Visited pak files. */
-		TSet<FString>& VisitedPakFiles;
-		/** Cached list of pak files. */
-		TArray<FPakListEntry>& Paks;
+		FDirectoryVisitor& Visitor;
 
 		/** Constructor. */
-		FPakVisitor(FDirectoryVisitor& InVisitor, TArray<FPakListEntry>& InPaks, TSet<FString>& InVisitedPakFiles)
-			: Visitor(InVisitor)
-			, VisitedPakFiles(InVisitedPakFiles)
-			, Paks(InPaks)
+		FPreventDuplicatesVisitor(FDirectoryVisitor& InVisitor, TSet<FString>& InVisitedFiles)
+			: FPreventDuplicatesVisitorBase(InVisitedFiles)
+			, Visitor(InVisitor)
 		{}
 		virtual bool Visit(const TCHAR* FilenameOrDirectory, bool bIsDirectory)
 		{
-			if (bIsDirectory == false)
+			if (CheckDuplicate(FilenameOrDirectory))
 			{
-				FString StandardFilename(FilenameOrDirectory);
-				FPaths::MakeStandardFilename(StandardFilename);
-
-				if (VisitedPakFiles.Contains(StandardFilename))
-				{
-					// Already visited, continue iterating.
-					return true;
-				}
-				else if (FPakPlatformFile::FindFileInPakFiles(Paks, FilenameOrDirectory, nullptr))
-				{
-					VisitedPakFiles.Add(StandardFilename);
-				}
-			}			
-			return Visitor.Visit(FilenameOrDirectory, bIsDirectory);
+				// Already visited, continue iterating.
+				return true;
+			}
+			return Visitor.Visit(*NormalizedFilename, bIsDirectory);
 		}
 	};
 
 	virtual bool IterateDirectory(const TCHAR* Directory, IPlatformFile::FDirectoryVisitor& Visitor) override
 	{
-		bool Result = true;
+		return IterateDirectoryInternal(Directory, Visitor, false /* bRecursive */);
+	}
+
+	bool IterateDirectoryInternal(const TCHAR* Directory, IPlatformFile::FDirectoryVisitor& Visitor, bool bRecursive)
+	{
+		TUniqueFunction<bool(const FString&, const FString&, bool, FPakFile&)> VisitFunction = [&Visitor](const FString& Filename, const FString& NormalizedFilename, bool bIsDir, FPakFile& PakFile)
+		{
+			return Visitor.Visit(*NormalizedFilename, bIsDir);
+		};
 		TSet<FString> FilesVisitedInPak;
+		bool Result = IterateDirectoryInternal(Directory, VisitFunction, bRecursive, FilesVisitedInPak);
+		if (Result && LowerLevel->DirectoryExists(Directory))
+		{
+			// Iterate inner filesystem but don't visit any files that were found in the Paks
+			FPreventDuplicatesVisitor PreventDuplicatesVisitor(Visitor, FilesVisitedInPak);
+			IPlatformFile::FDirectoryVisitor& LowerLevelVisitor(FilesVisitedInPak.Num() ? PreventDuplicatesVisitor : Visitor); // For performance, skip using PreventDuplicatedVisitor if there were no hits in pak
+			if (bRecursive)
+			{
+				Result = LowerLevel->IterateDirectoryRecursively(Directory, LowerLevelVisitor);
+			}
+			else
+			{
+				Result = LowerLevel->IterateDirectory(Directory, LowerLevelVisitor);
+			}
+		}
+		return Result;
+	}
+
+	bool IterateDirectoryInternal(const TCHAR* Directory, TUniqueFunction<bool(const FString&, const FString&, bool, FPakFile&)>& VisitFunction, bool bRecursive, TSet<FString>& FilesVisitedInPak)
+	{
+		bool Result = true;
 
 		TArray<FPakListEntry> Paks;
 		FString StandardDirectory = Directory;
@@ -2608,45 +2668,37 @@ public:
 		}
 
 		// Iterate pak files first
+		FString NormalizationBuffer;
+		TSet<FString> FilesVisitedInThisPak;
 		for (int32 PakIndex = 0; PakIndex < Paks.Num(); PakIndex++)
 		{
 			FPakFile& PakFile = *Paks[PakIndex].PakFile;
 			
 			const bool bIncludeFiles = true;
 			const bool bIncludeFolders = true;
-			TSet<FString> FilesVisitedInThisPak;
 
-			PakFile.FindPrunedFilesAtPath(FilesVisitedInThisPak, *StandardDirectory, bIncludeFiles, bIncludeFolders);
+			FilesVisitedInThisPak.Reset();
+			PakFile.FindPrunedFilesAtPath(FilesVisitedInThisPak, *StandardDirectory, bIncludeFiles, bIncludeFolders, bRecursive);
 			for (TSet<FString>::TConstIterator SetIt(FilesVisitedInThisPak); SetIt && Result; ++SetIt)
 			{
 				const FString& Filename = *SetIt;
-				if (!FilesVisitedInPak.Contains(Filename))
+				bool bIsDir = Filename.Len() && Filename[Filename.Len() - 1] == '/';
+				const FString* NormalizedFilename;
+				if (bIsDir)
 				{
-					bool bIsDir = Filename.Len() && Filename[Filename.Len() - 1] == '/';
-					if (bIsDir)
-					{
-						Result = Visitor.Visit(*Filename.LeftChop(1), true) && Result;
-					}
-					else
-					{
-						Result = Visitor.Visit(*Filename, false) && Result;
-					}
-					FilesVisitedInPak.Add(Filename);
+					NormalizationBuffer.Reset(Filename.Len());
+					NormalizationBuffer.AppendChars(*Filename, Filename.Len()-1); // Chop off the trailing /
+					NormalizedFilename = &NormalizationBuffer;
+				}
+				else
+				{
+					NormalizedFilename = &Filename;
+				}
+				if (!FilesVisitedInPak.Contains(*NormalizedFilename))
+				{
+					FilesVisitedInPak.Add(*NormalizedFilename);
+					Result = VisitFunction(Filename, *NormalizedFilename, bIsDir, PakFile) && Result;
 				}
-			}
-		}
-		if (Result && LowerLevel->DirectoryExists(Directory))
-		{
-			if (FilesVisitedInPak.Num())
-			{
-				// Iterate inner filesystem using FPakVisitor
-				FPakVisitor PakVisitor(Visitor, Paks, FilesVisitedInPak);
-				Result = LowerLevel->IterateDirectory(Directory, PakVisitor);
-			}
-			else
-			{
-				 // No point in using FPakVisitor as it will only slow things down.
-				Result = LowerLevel->IterateDirectory(Directory, Visitor);
 			}
 		}
 		return Result;
@@ -2654,131 +2706,76 @@ public:
 
 	virtual bool IterateDirectoryRecursively(const TCHAR* Directory, IPlatformFile::FDirectoryVisitor& Visitor) override
 	{
-		TSet<FString> FilesVisitedInPak;
-		TArray<FPakListEntry> Paks;
-		GetMountedPaks(Paks);
-		FPakVisitor PakVisitor(Visitor, Paks, FilesVisitedInPak);
-		return IPlatformFile::IterateDirectoryRecursively(Directory, PakVisitor);
+		return IterateDirectoryInternal(Directory, Visitor, true /* bRecursive */);
 	}
 
-	/**
-	 * Helper class to filter out files which have already been visited in one of the pak files.
-	 */
-	class FPakStatVisitor : public IPlatformFile::FDirectoryStatVisitor
+	class FPreventDuplicatesStatVisitor : public FPreventDuplicatesVisitorBase, public IPlatformFile::FDirectoryStatVisitor
 	{
 	public:
 		/** Wrapped visitor. */
-		FDirectoryStatVisitor&	Visitor;
-		/** Visited pak files. */
-		TSet<FString>& VisitedPakFiles;
-		/** Cached list of pak files. */
-		TArray<FPakListEntry>& Paks;
+		FDirectoryStatVisitor& Visitor;
 
 		/** Constructor. */
-		FPakStatVisitor(FDirectoryStatVisitor& InVisitor, TArray<FPakListEntry>& InPaks, TSet<FString>& InVisitedPakFiles)
-			: Visitor(InVisitor)
-			, VisitedPakFiles(InVisitedPakFiles)
-			, Paks(InPaks)
+		FPreventDuplicatesStatVisitor(FDirectoryStatVisitor& InVisitor, TSet<FString>& InVisitedFiles)
+			: FPreventDuplicatesVisitorBase(InVisitedFiles)
+			, Visitor(InVisitor)
 		{}
 		virtual bool Visit(const TCHAR* FilenameOrDirectory, const FFileStatData& StatData)
 		{
-			if (StatData.bIsDirectory == false)
+			if (CheckDuplicate(FilenameOrDirectory))
 			{
-				FString StandardFilename(FilenameOrDirectory);
-				FPaths::MakeStandardFilename(StandardFilename);
-
-				if (VisitedPakFiles.Contains(StandardFilename))
-				{
-					// Already visited, continue iterating.
-					return true;
-				}
-				else if (FPakPlatformFile::FindFileInPakFiles(Paks, FilenameOrDirectory, nullptr))
-				{
-					VisitedPakFiles.Add(StandardFilename);
-				}
-			}			
-			return Visitor.Visit(FilenameOrDirectory, StatData);
+				// Already visited, continue iterating.
+				return true;
+			}
+			return Visitor.Visit(*NormalizedFilename, StatData);
 		}
 	};
 
 	virtual bool IterateDirectoryStat(const TCHAR* Directory, IPlatformFile::FDirectoryStatVisitor& Visitor) override
 	{
-		bool Result = true;
-		TSet<FString> FilesVisitedInPak;
-
-		TArray<FPakListEntry> Paks;
- 
-		FString StandardDirectory = Directory;
-		FPaths::MakeStandardFilename(StandardDirectory);
-
-		bool bIsDownloadableDir = (FPaths::HasProjectPersistentDownloadDir() && StandardDirectory.StartsWith(FPaths::ProjectPersistentDownloadDir())) || StandardDirectory.StartsWith(FPaths::CloudDir());
-
-		// don't look for in pak files for target-only locations
-		if (!bIsDownloadableDir)
-		{
-			GetMountedPaks(Paks);
-		}
+		return IterateDirectoryStatInternal(Directory, Visitor, false /* bRecursive */);
+	}
 
-		// Iterate pak files first
-		for (int32 PakIndex = 0; PakIndex < Paks.Num(); PakIndex++)
+	bool IterateDirectoryStatInternal(const TCHAR* Directory, IPlatformFile::FDirectoryStatVisitor& Visitor, bool bRecursive)
+	{
+		TUniqueFunction<bool(const FString&, const FString&, bool, FPakFile&)> VisitFunction = [&Visitor, this](const FString& Filename, const FString& NormalizedFilename, bool bIsDir, FPakFile& PakFile)
 		{
-			FPakFile& PakFile = *Paks[PakIndex].PakFile;
-			
-			const bool bIncludeFiles = true;
-			const bool bIncludeFolders = true;
-			TSet<FString> FilesVisitedInThisPak;
-
-			PakFile.FindPrunedFilesAtPath(FilesVisitedInThisPak, *StandardDirectory, bIncludeFiles, bIncludeFolders);
-			for (TSet<FString>::TConstIterator SetIt(FilesVisitedInThisPak); SetIt && Result; ++SetIt)
+			int64 FileSize = -1;
+			if (!bIsDir)
 			{
-				const FString& Filename = *SetIt;
-				if (!FilesVisitedInPak.Contains(Filename))
+				FPakEntry FileEntry;
+				if (FindFileInPakFiles(*Filename, nullptr, &FileEntry))
 				{
-					bool bIsDir = Filename.Len() && Filename[Filename.Len() - 1] == '/';
+					FileSize = (FileEntry.CompressionMethodIndex != 0) ? FileEntry.UncompressedSize : FileEntry.Size;
+				}
+			}
 
-					int64 FileSize = -1;
-					if (!bIsDir)
-					{
-						FPakEntry FileEntry;
-						if (FindFileInPakFiles(*Filename, nullptr, &FileEntry))
-						{
-							FileSize = (FileEntry.CompressionMethodIndex != 0) ? FileEntry.UncompressedSize : FileEntry.Size;
-						}
-					}
+			const FFileStatData StatData(
+				PakFile.GetTimestamp(),
+				PakFile.GetTimestamp(),
+				PakFile.GetTimestamp(),
+				FileSize,
+				bIsDir,
+				true	// IsReadOnly
+			);
 
-					const FFileStatData StatData(
-						PakFile.GetTimestamp(),
-						PakFile.GetTimestamp(),
-						PakFile.GetTimestamp(),
-						FileSize, 
-						bIsDir,
-						true	// IsReadOnly
-						);
+			return Visitor.Visit(*NormalizedFilename, StatData);
+		};
 
-					if (bIsDir)
-					{
-						Result = Visitor.Visit(*Filename.LeftChop(1), StatData) && Result;
-					}
-					else
-					{
-						Result = Visitor.Visit(*Filename, StatData) && Result;
-					}
-					FilesVisitedInPak.Add(Filename);
-				}
-			}
-		}
+		TSet<FString> FilesVisitedInPak;
+		bool Result = IterateDirectoryInternal(Directory, VisitFunction, bRecursive, FilesVisitedInPak);
 		if (Result && LowerLevel->DirectoryExists(Directory))
 		{
-			if (FilesVisitedInPak.Num())
+			// Iterate inner filesystem but don't visit any files that were found in the Paks
+			FPreventDuplicatesStatVisitor PreventDuplicatesVisitor(Visitor, FilesVisitedInPak);
+			IPlatformFile::FDirectoryStatVisitor& LowerLevelVisitor(FilesVisitedInPak.Num() ? PreventDuplicatesVisitor : Visitor); // For performance, skip using PreventDuplicatedVisitor if there were no hits in pak
+			if (bRecursive)
 			{
-				// Iterate inner filesystem using FPakVisitor
-				FPakStatVisitor PakVisitor(Visitor, Paks, FilesVisitedInPak);
-				Result = LowerLevel->IterateDirectoryStat(Directory, PakVisitor);
+				Result = LowerLevel->IterateDirectoryStatRecursively(Directory, LowerLevelVisitor);
 			}
 			else
 			{
-				 // No point in using FPakVisitor as it will only slow things down.
-				Result = LowerLevel->IterateDirectoryStat(Directory, Visitor);
+				Result = LowerLevel->IterateDirectoryStat(Directory, LowerLevelVisitor);
 			}
 		}
 		return Result;
@@ -2786,11 +2783,7 @@ public:
 
 	virtual bool IterateDirectoryStatRecursively(const TCHAR* Directory, IPlatformFile::FDirectoryStatVisitor& Visitor) override
 	{
-		TSet<FString> FilesVisitedInPak;
-		TArray<FPakListEntry> Paks;
-		GetMountedPaks(Paks);
-		FPakStatVisitor PakVisitor(Visitor, Paks, FilesVisitedInPak);
-		return IPlatformFile::IterateDirectoryStatRecursively(Directory, PakVisitor);
+		return IterateDirectoryStatInternal(Directory, Visitor, true/* bRecursive */);
 	}
 
 	virtual void FindFiles(TArray<FString>& FoundFiles, const TCHAR* Directory, const TCHAR* FileExtension) override
