Pakfile pruning - disable by default the stripping of directory names; this was the default in 4.25 and was accidentally changed when changing pakfile indexes to strip directory names on load.

diff --git a/Engine/Source/Runtime/PakFile/Private/IPlatformFilePak.cpp b/Engine/Source/Runtime/PakFile/Private/IPlatformFilePak.cpp
--- a/Engine/Source/Runtime/PakFile/Private/IPlatformFilePak.cpp
+++ b/Engine/Source/Runtime/PakFile/Private/IPlatformFilePak.cpp
@@ -5635,7 +5635,7 @@ struct FPakFile::FIndexSettings
 {
 	FIndexSettings()
 	{
-		bKeepFullDirectory = false;
+		bKeepFullDirectory = true;
 		bValidatePruning = false;
 		bDelayPruning = false;
 		bWritePathHashIndex = true;
@@ -5653,7 +5653,12 @@ struct FPakFile::FIndexSettings
 		GConfig->GetBool(TEXT("Pak"), TEXT("WriteFullDirectoryIndex"), bWriteFullDirectoryIndex, GEngineIni);
 #endif
 
-		bKeepFullDirectory = bKeepFullDirectory || IS_PROGRAM || !FPlatformProperties::RequiresCookedData();
+#if IS_PROGRAM || WITH_EDITOR
+		// Directory pruning is not enabled in the editor or in development programs because there is no need to save the memory in those environments and some development features require not pruning
+		bKeepFullDirectory = true;
+#else
+		bKeepFullDirectory = bKeepFullDirectory || !FPlatformProperties::RequiresCookedData();
+#endif
 #if !UE_BUILD_SHIPPING
 		const TCHAR* CommandLine = FCommandLine::Get();
 		FParse::Bool(CommandLine, TEXT("ForcePakKeepFullDirectory="), bKeepFullDirectory);
@@ -6799,6 +6804,20 @@ int32 GetRecursiveAllocatedSize(const FPakFile::FDirectoryIndex& Index)
 
 void FPakPlatformFile::OptimizeMemoryUsageForMountedPaks()
 {
+#if !UE_BUILD_SHIPPING
+	// UE_DEPRECATED(4.26, "UnloadPakEntryFilenamesIfPossible is deprecated.")
+	bool bUnloadPakEntryFilenamesIfPossible = false;
+	GConfig->GetBool(TEXT("Pak"), TEXT("UnloadPakEntryFilenamesIfPossible"), bUnloadPakEntryFilenamesIfPossible, GEngineIni);
+	if (bUnloadPakEntryFilenamesIfPossible)
+	{
+		UE_LOG(LogPakFile, Warning, TEXT("The UnloadPakEntryFilenamesIfPossible has been deprecated and is no longer sufficient to specify the unloading of pak files.\n")
+			TEXT("The choice to not load pak files is now made earlier than Ini settings are available.\n")
+			TEXT("To specify that filenames should be removed from the runtime PakFileIndex, use the new runtime delegate FPakPlatformFile::GetPakSetIndexSettingsDelegate().\n")
+			TEXT("In a global variable constructor that executes before the process main function, bind this delegate to a function that sets the output bool bKeepFullDirectory to false.\n")
+			TEXT("See FShooterPreMainCallbacks in Samples\\Games\\ShooterGame\\Source\\ShooterGame\\Private\\ShooterGameModule.cpp for an example binding."));
+	}
+#endif
+
 #if ENABLE_PAKFILE_RUNTIME_PRUNING || !UE_BUILD_SHIPPING
 	TArray<FPakListEntry> Paks;
 	bool bNeedsPaks = false;
