Fix a crash when decryption of a pak signature fails

diff --git a/Engine/Source/Runtime/PakFile/Public/IPlatformFilePak.h b/Engine/Source/Runtime/PakFile/Public/IPlatformFilePak.h
--- a/Engine/Source/Runtime/PakFile/Public/IPlatformFilePak.h
+++ b/Engine/Source/Runtime/PakFile/Public/IPlatformFilePak.h
@@ -3081,7 +3081,7 @@ struct FPakSignatureFile
 		else
 		{
 			int32 BytesDecrypted = FRSA::DecryptPublic(EncryptedHash, SignatureData, InKey);
-			if (BytesDecrypted > UE_ARRAY_COUNT(FSHAHash::Hash))
+			if (BytesDecrypted > (int32)UE_ARRAY_COUNT(FSHAHash::Hash))
 			{
 				FMemory::Memcpy(DecryptedHash.Hash, SignatureData.GetData() + SignatureData.Num() - UE_ARRAY_COUNT(FSHAHash::Hash), UE_ARRAY_COUNT(FSHAHash::Hash));
 				SignatureData.SetNum(SignatureData.Num() - UE_ARRAY_COUNT(FSHAHash::Hash));
