Fix EncodePakEntry for compressionblocksize.  Maintains binary compatibility with existing paks.

diff --git a/Engine/Source/Runtime/PakFile/Private/IPlatformFilePak.cpp b/Engine/Source/Runtime/PakFile/Private/IPlatformFilePak.cpp
--- a/Engine/Source/Runtime/PakFile/Private/IPlatformFilePak.cpp
+++ b/Engine/Source/Runtime/PakFile/Private/IPlatformFilePak.cpp
@@ -6030,10 +6030,6 @@ bool FPakFile::EncodePakEntry(FArchive& Ar, const FPakEntry& InPakEntry, const F
 	}
 	if (PakEntry.CompressionMethodIndex != 0)
 	{
-		if (PakEntry.CompressionBlockSize != PakEntry.UncompressedSize && ((PakEntry.CompressionBlockSize >> 11) > 0x3f))
-		{
-			return false;
-		}
 		if (PakEntry.CompressionBlocks.Num() > 0 && ((InInfo.HasRelativeCompressedChunkOffsets() ? 0 : PakEntry.Offset) + HeaderSize != PakEntry.CompressionBlocks[0].CompressedStart))
 		{
 			return false;
@@ -6066,6 +6062,14 @@ bool FPakFile::EncodePakEntry(FArchive& Ar, const FPakEntry& InPakEntry, const F
 	bool bIsOffset32BitSafe = PakEntry.Offset <= MAX_uint32;
 	bool bIsSize32BitSafe = PakEntry.Size <= MAX_uint32;
 	bool bIsUncompressedSize32BitSafe = PakEntry.UncompressedSize <= MAX_uint32;
+	
+	// compression block size is sent in the flags field for backwards compatibility (eg. when it is 64k)
+	// if not the maximum value of the packed field is sent as a flag (0x3F)
+	uint32 CompressionBlockSizePacked = (PakEntry.CompressionBlockSize >> 11) & 0x3F;
+	if ( (CompressionBlockSizePacked<<11) != PakEntry.CompressionBlockSize )
+	{
+		CompressionBlockSizePacked = 0x3F;
+	}
 
 	// Build the Flags field.
 	uint32 Flags =
@@ -6075,10 +6079,17 @@ bool FPakFile::EncodePakEntry(FArchive& Ar, const FPakEntry& InPakEntry, const F
 		| (PakEntry.CompressionMethodIndex << 23)
 		| (PakEntry.IsEncrypted() ? (1 << 22) : 0)
 		| (PakEntry.CompressionBlocks.Num() << 6)
-		| (PakEntry.CompressionBlockSize >> 11)
+		| CompressionBlockSizePacked
 		;
 
 	Ar << Flags;
+	
+	// if we write 0x3F for CompressionBlockSize then send the field
+	if ( CompressionBlockSizePacked == 0x3F )
+	{
+		uint32 Value = (uint32)PakEntry.CompressionBlockSize;
+		Ar << Value;
+	}
 
 	// Build the Offset field.
 	if (bIsOffset32BitSafe)
@@ -6142,6 +6153,18 @@ void FPakFile::DecodePakEntry(const uint8* SourcePtr, FPakEntry& OutEntry, const
 	// Bits 5-0 = Compression block size
 	uint32 Value = *(uint32*)SourcePtr;
 	SourcePtr += sizeof(uint32);
+	
+	uint32 CompressionBlockSize = 0;
+	if ( (Value & 0x3f) == 0x3f ) // flag value to load a field
+	{
+		CompressionBlockSize = *(uint32*)SourcePtr;
+		SourcePtr += sizeof(uint32);
+	}
+	else
+	{
+		// for backwards compatibility with old paks :
+		CompressionBlockSize = ((Value & 0x3f) << 11);
+	}
 
 	// Filter out the CompressionMethod.
 	OutEntry.CompressionMethodIndex = (Value >> 23) & 0x3f;
@@ -6208,11 +6231,16 @@ void FPakFile::DecodePakEntry(const uint8* SourcePtr, FPakEntry& OutEntry, const
 	OutEntry.CompressionBlocks.Empty(CompressionBlocksCount);
 	OutEntry.CompressionBlocks.SetNum(CompressionBlocksCount);
 
-	// Filter the compression block size or use the UncompressedSize if less that 64k.
+	// Set CompressionBlockSize with conditional clamping
+	// this is probably not necessary but maintains past behavior
 	OutEntry.CompressionBlockSize = 0;
 	if (CompressionBlocksCount > 0)
 	{
-		OutEntry.CompressionBlockSize = OutEntry.UncompressedSize < 65536 ? (uint32)OutEntry.UncompressedSize : ((Value & 0x3f) << 11);
+		OutEntry.CompressionBlockSize = CompressionBlockSize;
+		if ( OutEntry.UncompressedSize < CompressionBlockSize )
+		{
+			OutEntry.CompressionBlockSize = OutEntry.UncompressedSize;
+		}
 	}
 
 	// Set bDeleteRecord to false, because it obviously isn't deleted if we are here.
