U5 - Don't keep the lock across the waitpoint.

diff --git a/Engine/Source/Runtime/PakFile/Private/IPlatformFilePak.cpp b/Engine/Source/Runtime/PakFile/Private/IPlatformFilePak.cpp
--- a/Engine/Source/Runtime/PakFile/Private/IPlatformFilePak.cpp
+++ b/Engine/Source/Runtime/PakFile/Private/IPlatformFilePak.cpp
@@ -2516,12 +2516,19 @@ private: // below here we assume CachedFilesScopeLock until we get to the next s
 	{
 		if (!NotifyRecursion)
 		{
-			for (IAsyncReadRequest* Elem : RequestsToDelete)
+			TArray<IAsyncReadRequest*> Swapped;
+			{
+				FScopeLock Lock(&CachedFilesScopeLock);
+				Swapped = (MoveTemp(RequestsToDelete));
+				check(RequestsToDelete.IsEmpty());
+			}
+
+			for (IAsyncReadRequest* Elem : Swapped)
 			{
 				Elem->WaitCompletion();
 				delete Elem;
 			}
-			RequestsToDelete.Empty();
+			Swapped.Empty();
 		}
 	}
 	void StartBlockTask(FCacheBlock& Block)
@@ -2834,9 +2841,10 @@ public:
 	void NewRequestsToLowerComplete(bool bWasCanceled, IAsyncReadRequest* Request, int32 Index)
 	{
 		LLM_SCOPE(ELLMTag::FileSystem);
+		ClearOldBlockTasks();
+
 		FScopeLock Lock(&CachedFilesScopeLock);
 		RequestsToLower[Index].RequestHandle = Request;
-		ClearOldBlockTasks();
 		NotifyRecursion++;
 		if (!RequestsToLower[Index].Memory) // might have already been filled in by the signature check
 		{
@@ -2926,8 +2934,9 @@ public:
 	bool GetCompletedRequest(IPakRequestor* Owner, uint8* UserSuppliedMemory)
 	{
 		check(Owner);
-		FScopeLock Lock(&CachedFilesScopeLock);
 		ClearOldBlockTasks();
+
+		FScopeLock Lock(&CachedFilesScopeLock);
 		TIntervalTreeIndex RequestIndex = OutstandingRequests.FindRef(Owner->UniqueID);
 		static_assert(IntervalTreeInvalidIndex == 0, "FindRef will return 0 for something not found");
 		if (RequestIndex)
@@ -2942,8 +2951,9 @@ public:
 	void CancelRequest(IPakRequestor* Owner)
 	{
 		check(Owner);
-		FScopeLock Lock(&CachedFilesScopeLock);
 		ClearOldBlockTasks();
+
+		FScopeLock Lock(&CachedFilesScopeLock);
 		TIntervalTreeIndex RequestIndex = OutstandingRequests.FindRef(Owner->UniqueID);
 		static_assert(IntervalTreeInvalidIndex == 0, "FindRef will return 0 for something not found");
 		if (RequestIndex)
