Add recursion support to FCompressionScratchBuffers. This is needed to fix a race condition with the inflight decompression task. It wasn�t a problem in UE4 because FAsyncTask::SyncCompletion used to just wait on a completion event. In UE5, however, it will try to pull other tasks from task queues and execute them locally until either there is no more task or the completion event is triggered.

diff --git a/Engine/Source/Runtime/PakFile/Private/IPlatformFilePak.cpp b/Engine/Source/Runtime/PakFile/Private/IPlatformFilePak.cpp
--- a/Engine/Source/Runtime/PakFile/Private/IPlatformFilePak.cpp
+++ b/Engine/Source/Runtime/PakFile/Private/IPlatformFilePak.cpp
@@ -4669,17 +4669,14 @@ public:
 	}
 };
 
-/**
- * Thread local class to manage working buffers for file compression
- */
-class FCompressionScratchBuffers : public TThreadSingleton<FCompressionScratchBuffers>
+struct FCompressionScratchBuffers
 {
-public:
 	FCompressionScratchBuffers()
 		: TempBufferSize(0)
 		, ScratchBufferSize(0)
 		, LastPakEntryOffset(-1)
 		, LastDecompressedBlock(0xFFFFFFFF)
+		, Next(nullptr)
 	{}
 
 	int64				TempBufferSize;
@@ -4691,6 +4688,8 @@ public:
 	FSHAHash LastPakIndexHash;
 	uint32 LastDecompressedBlock;
 
+	FCompressionScratchBuffers* Next;
+
 	void EnsureBufferSpace(int64 CompressionBlockSize, int64 ScrachSize)
 	{
 		if (TempBufferSize < CompressionBlockSize)
@@ -4706,6 +4705,75 @@ public:
 	}
 };
 
+/**
+ * Thread local class to manage working buffers for file compression
+ */
+class FCompressionScratchBuffersStack : public TThreadSingleton<FCompressionScratchBuffersStack>
+{
+public:
+	FCompressionScratchBuffersStack()
+		: bFirstInUse(false)
+		, RecursionList(nullptr)
+	{}
+
+private:
+	FCompressionScratchBuffers* Acquire()
+	{
+		if (!bFirstInUse)
+		{
+			bFirstInUse = true;
+			return &First;
+		}
+		FCompressionScratchBuffers* Top = new FCompressionScratchBuffers;
+		Top->Next = RecursionList;
+		RecursionList = Top;
+		return Top;
+	}
+
+	void Release(FCompressionScratchBuffers* Top)
+	{
+		check(bFirstInUse);
+		if (!RecursionList)
+		{
+			check(Top == &First);
+			bFirstInUse = false;
+		}
+		else
+		{
+			check(Top == RecursionList);
+			RecursionList = Top->Next;
+			delete Top;
+		}
+	}
+
+	bool bFirstInUse;
+	FCompressionScratchBuffers First;
+	FCompressionScratchBuffers* RecursionList;
+
+	friend class FScopedCompressionScratchBuffers;
+};
+
+class FScopedCompressionScratchBuffers
+{
+public:
+	FScopedCompressionScratchBuffers()
+		: Inner(FCompressionScratchBuffersStack::Get().Acquire())
+	{}
+
+	~FScopedCompressionScratchBuffers()
+	{
+		FCompressionScratchBuffersStack::Get().Release(Inner);
+	}
+
+	FCompressionScratchBuffers* operator->() const
+	{
+		return Inner;
+	}
+
+private:
+	FCompressionScratchBuffers* Inner;
+};
+
 /**
  * Class to handle correctly reading from a compressed file within a pak
  */
@@ -4776,7 +4844,7 @@ public:
 		uint8* WorkingBuffers[2];
 		int64 DirectCopyStart = DesiredPosition % PakEntry.CompressionBlockSize;
 		FAsyncTask<FPakUncompressTask> UncompressTask;
-		FCompressionScratchBuffers& ScratchSpace = FCompressionScratchBuffers::Get();
+		FScopedCompressionScratchBuffers ScratchSpace;
 		bool bStartedUncompress = false;
 
 		FName CompressionMethod = PakFile.GetInfo().GetCompressionMethod(PakEntry.CompressionMethodIndex);
@@ -4792,10 +4860,10 @@ public:
 		float SlopMultiplier = 1.1f;
 		int64 WorkingBufferRequiredSize = FCompression::CompressMemoryBound(CompressionMethod, CompressionBlockSize) * SlopMultiplier;
 		WorkingBufferRequiredSize = EncryptionPolicy::AlignReadRequest(WorkingBufferRequiredSize);
-		const bool bExistingScratchBufferValid = ScratchSpace.TempBufferSize >= CompressionBlockSize;
-		ScratchSpace.EnsureBufferSpace(CompressionBlockSize, WorkingBufferRequiredSize * 2);
-		WorkingBuffers[0] = ScratchSpace.ScratchBuffer.Get();
-		WorkingBuffers[1] = ScratchSpace.ScratchBuffer.Get() + WorkingBufferRequiredSize;
+		const bool bExistingScratchBufferValid = ScratchSpace->TempBufferSize >= CompressionBlockSize;
+		ScratchSpace->EnsureBufferSpace(CompressionBlockSize, WorkingBufferRequiredSize * 2);
+		WorkingBuffers[0] = ScratchSpace->ScratchBuffer.Get();
+		WorkingBuffers[1] = ScratchSpace->ScratchBuffer.Get() + WorkingBufferRequiredSize;
 
 		FArchive* PakReader = AcquirePakReader();
 
@@ -4818,16 +4886,16 @@ public:
 			const bool bCurrentScratchTempBufferValid = 
 				bExistingScratchBufferValid && !bStartedUncompress
 				// ensure this object was the last reader from the scratch buffer and the last thing it decompressed was this block.
-				&& (ScratchSpace.LastPakEntryOffset == PakEntry.Offset)
-				&& (ScratchSpace.LastPakIndexHash == PakFile.GetInfo().IndexHash)
-				&& (ScratchSpace.LastDecompressedBlock == CompressionBlockIndex)
+				&& (ScratchSpace->LastPakEntryOffset == PakEntry.Offset)
+				&& (ScratchSpace->LastPakIndexHash == PakFile.GetInfo().IndexHash)
+				&& (ScratchSpace->LastDecompressedBlock == CompressionBlockIndex)
 				// ensure the previous decompression destination was the scratch buffer.
 				&& !(DirectCopyStart == 0 && Length >= CompressionBlockSize); 
 
 			if (bCurrentScratchTempBufferValid)
 			{
 				// Reuse the existing scratch buffer to avoid repeatedly deserializing and decompressing the same block.
-				FMemory::Memcpy(V, ScratchSpace.TempBuffer.Get() + DirectCopyStart, WriteSize);
+				FMemory::Memcpy(V, ScratchSpace->TempBuffer.Get() + DirectCopyStart, WriteSize);
 			}
 			else
 			{
@@ -4851,24 +4919,24 @@ public:
 					TaskDetails.CompressedBuffer = WorkingBuffers[CompressionBlockIndex & 1];
 					TaskDetails.CompressedSize = CompressedBlockSize;
 					TaskDetails.CopyOut = nullptr;
-					ScratchSpace.LastDecompressedBlock = 0xFFFFFFFF;
-					ScratchSpace.LastPakIndexHash = FSHAHash();
-					ScratchSpace.LastPakEntryOffset = -1;
+					ScratchSpace->LastDecompressedBlock = 0xFFFFFFFF;
+					ScratchSpace->LastPakIndexHash = FSHAHash();
+					ScratchSpace->LastPakEntryOffset = -1;
 				}
 				else
 				{
 					// Block needs to be copied from a working buffer
 					TaskDetails.CompressionFormat = CompressionMethod;
-					TaskDetails.UncompressedBuffer = ScratchSpace.TempBuffer.Get();
+					TaskDetails.UncompressedBuffer = ScratchSpace->TempBuffer.Get();
 					TaskDetails.UncompressedSize = UncompressedBlockSize;
 					TaskDetails.CompressedBuffer = WorkingBuffers[CompressionBlockIndex & 1];
 					TaskDetails.CompressedSize = CompressedBlockSize;
 					TaskDetails.CopyOut = V;
 					TaskDetails.CopyOffset = DirectCopyStart;
 					TaskDetails.CopyLength = WriteSize;
-					ScratchSpace.LastDecompressedBlock = CompressionBlockIndex;
-					ScratchSpace.LastPakIndexHash = PakFile.GetInfo().IndexHash;
-					ScratchSpace.LastPakEntryOffset = PakEntry.Offset;
+					ScratchSpace->LastDecompressedBlock = CompressionBlockIndex;
+					ScratchSpace->LastPakIndexHash = PakFile.GetInfo().IndexHash;
+					ScratchSpace->LastPakEntryOffset = PakEntry.Offset;
 				}
 
 				if (Length == WriteSize)
