FPakFile::DecodePakEntry: Fix special case handling for when CompressionBlockSize should be set equal to UncompressedSize.
The original legacy condition was to set CompressionBlockSize = UnCompressedSize when UncompressedSize < 65536.
That condition is incompatible with future plans for more flexible CompressionBlockSize.
We changed that to UnCompressedSize < CompressionBlockSize, but now we have found that
new check does not catch the case of e.g. UnCompressedSize == 909, CompressionBlockSize = 909, and encoded CompressionBlockSize == 909 >> 11 == 0.
The original legacy condition relied on the assumption that (UncompressedSize < 65536) implies (NumCompressionBlocks == 1).
In our second attempt at fixing the legacy condition, we now use the condition NumCompressionBlocks == 1.
#rnx

diff --git a/Engine/Source/Runtime/PakFile/Private/IPlatformFilePak.cpp b/Engine/Source/Runtime/PakFile/Private/IPlatformFilePak.cpp
--- a/Engine/Source/Runtime/PakFile/Private/IPlatformFilePak.cpp
+++ b/Engine/Source/Runtime/PakFile/Private/IPlatformFilePak.cpp
@@ -6153,12 +6153,18 @@ bool FPakFile::EncodePakEntry(FArchive& Ar, const FPakEntry& InPakEntry, const F
 	bool bIsSize32BitSafe = PakEntry.Size <= MAX_uint32;
 	bool bIsUncompressedSize32BitSafe = PakEntry.UncompressedSize <= MAX_uint32;
 	
-	// compression block size is sent in the flags field when it is 64k
-	// if not the maximum value of the packed field is sent as a flag (0x3F)
-	uint32 CompressionBlockSizePacked = (PakEntry.CompressionBlockSize >> 11) & 0x3F;
-	if ( (CompressionBlockSizePacked<<11) != PakEntry.CompressionBlockSize )
+	// If CompressionBlocks.Num() == 1, we set CompressionBlockSize == UncompressedSize and record CompressBlockSizePacked=0
+	// Otherwise, we encode CompressionBlockSize as a 6-bit multiple of 1 << 11.
+	// If CompressionBlockSize is not a multiple of 1 << 11, or is a larger multiple than 6 bits we can not encode correctly.
+	// In that case we set the packed field to its maximum value (0x3F) and send the unencoded CompressionBlockSize as a separate value.
+	uint32 CompressionBlockSizePacked = 0;
+	if (PakEntry.CompressionBlocks.Num() > 1)
 	{
-		CompressionBlockSizePacked = 0x3F;
+		CompressionBlockSizePacked = (PakEntry.CompressionBlockSize >> 11) & 0x3F;
+		if ((CompressionBlockSizePacked << 11) != PakEntry.CompressionBlockSize)
+		{
+			CompressionBlockSizePacked = 0x3F;
+		}
 	}
 
 	// Build the Flags field.
@@ -6321,16 +6327,16 @@ void FPakFile::DecodePakEntry(const uint8* SourcePtr, FPakEntry& OutEntry, const
 	OutEntry.CompressionBlocks.Empty(CompressionBlocksCount);
 	OutEntry.CompressionBlocks.SetNum(CompressionBlocksCount);
 	
-	// Set CompressionBlockSize with conditional clamping
-	// this is probably not necessary but maintains past behavior
 	OutEntry.CompressionBlockSize = 0;
 	if (CompressionBlocksCount > 0)
 	{
 		OutEntry.CompressionBlockSize = CompressionBlockSize;
-		if ( OutEntry.UncompressedSize < CompressionBlockSize )
+		// Per the comment in Encode, if CompressionBlocksCount == 1, we use UncompressedSize for CompressionBlockSize
+		if (CompressionBlocksCount == 1)
 		{
 			OutEntry.CompressionBlockSize = OutEntry.UncompressedSize;
 		}
+		ensure(OutEntry.CompressionBlockSize != 0);
 	}
 
 	// Set bDeleteRecord to false, because it obviously isn't deleted if we are here.
