Add an earlier FATAL log to catch invalid premature deletion of FPakProcessedReadRequest before it can cause difficult-to-diagnose crashes on the asynchronous DoProcessing thread.
#rnx

diff --git a/Engine/Source/Runtime/PakFile/Private/IPlatformFilePak.cpp b/Engine/Source/Runtime/PakFile/Private/IPlatformFilePak.cpp
--- a/Engine/Source/Runtime/PakFile/Private/IPlatformFilePak.cpp
+++ b/Engine/Source/Runtime/PakFile/Private/IPlatformFilePak.cpp
@@ -3518,6 +3518,7 @@ public:
 
 	virtual ~FPakProcessedReadRequest()
 	{
+		UE_CLOG(!bCompleteAndCallbackCalled, LogPakFile, Fatal, TEXT("IAsyncReadRequests must not be deleted until they are completed."));
 		check(!MyCanceledBlocks.Num());
 		DoneWithRawRequests();
 		if (Memory && !bUserSuppliedMemory)
@@ -4205,6 +4206,11 @@ public:
 			check(Block.RefCount > 0);
 			if (!--Block.RefCount)
 			{
+				// If this no-longer-referenced block is still held by the RawReadThread+DoProcessing functions, DoProcessing will crash when it assumes
+				// the block has been canceled and assumes it is present in OutstandingCancelMapBlock
+				// We have to fatally assert instead of doing what cancel does, because RemoveRequest is called from the request's destructor,
+				// and therefore we don't have a persistent request that can take responsibility for staying alive until the canceled block finishes processing
+				UE_CLOG(Block.bInFlight && !Block.bCPUWorkIsComplete, LogPakFile, Fatal, TEXT("RemoveRequest called on Request that still has a block in processing."));
 				if (Block.RawRequest)
 				{
 					Block.RawRequest->Cancel();
