Add comments and a test framework for testing FPakProcessedReadRequest.
#rnx

diff --git a/Engine/Source/Runtime/PakFile/Private/IPlatformFilePak.cpp b/Engine/Source/Runtime/PakFile/Private/IPlatformFilePak.cpp
--- a/Engine/Source/Runtime/PakFile/Private/IPlatformFilePak.cpp
+++ b/Engine/Source/Runtime/PakFile/Private/IPlatformFilePak.cpp
@@ -455,6 +455,19 @@ static FAutoConsoleVariableRef CVar_ForceDecompressionFails(
 	GPakCache_ForceDecompressionFails,
 	TEXT("If > 0, then force decompression failures to test the panic sync read fallback.")
 );
+static bool GPakCache_ForcePakProcessedReads = false;
+static FAutoConsoleVariableRef CVar_ForcePakProcessReads(
+	TEXT("ForcePakProcessReads"),
+	GPakCache_ForcePakProcessedReads,
+	TEXT("If true, then Asynchronous reads from pak files will always used the FPakProcessedReadRequest system that is ordinarily only used on compressed files.")
+);
+static bool GetPakCacheForcePakProcessedReads()
+{
+	static bool bInitialValue = (GPakCache_ForcePakProcessedReads = FParse::Param(FCommandLine::Get(), TEXT("ForcePakProcessReads")));
+	return GPakCache_ForcePakProcessedReads;
+}
+static FName GPakFakeCompression(TEXT("PakFakeCompression"));
+
 #endif
 
 class FPakSizeRequest : public IAsyncReadRequest
@@ -3200,17 +3213,49 @@ class FPakAsyncReadFileHandle;
 
 struct FCachedAsyncBlock
 {
+	/**
+	 * Assigned in FPakAsyncReadFileHandle::StartBlock to store the handle for the raw read request.
+	 * Readable only under FPakAsyncReadFileHandle->CriticalSection, or from RawReadCallback.
+	 * Can not be written under CriticalSection until after RawRequest->WaitCompletion.
+	 * Set to null under critical section from DoProcessing or from cancelation.
+	 */
 	class FPakReadRequest* RawRequest;
-	uint8* Raw; // compressed, encrypted and/or signature not checked
-	uint8* Processed; // decompressed, deencrypted and signature checked
+	/**
+	 * compressed, encrypted and/or signature not checked
+	 * Set to null in FPakAsyncReadFileHandle::StartBlock. RawReadRequest and DoProcessing can assign
+	 * and modify it outside of FPakAsyncReadFileHandle->CriticalSection.
+	 * Can not be read/written by any other thread until RawRequest is set to null and bCPUWorkIsComplete is set to false.
+	 */
+	uint8* Raw;
+	/** decompressed, deencrypted and signature checked */
+	uint8* Processed;
 	FGraphEventRef CPUWorkGraphEvent;
 	int32 RawSize;
 	int32 DecompressionRawSize;
 	int32 ProcessedSize;
+	/**
+	 * How many requests touch the block that are still alive and uncanceled. Accessed only within FPakAsyncReadFileHandle->CriticalSection.
+	 * When the reference count goes to 0, the block is removed from Blocks, but async threads might still have a pointer to it.
+	 * Block is deleted when refcount is 0 and async thread has finished with it (bCPUWorkIsComplete =true).
+	 */
 	int32 RefCount;
 	int32 BlockIndex;
+	/**
+	 * The block has been requested, and is still referenced, either from still-alive requests or from the async load and processing  of the block.
+	 * Accessed only within FPakAsyncReadFileHandle->CriticalSection. Modified when requests start, cancel/destroy, and when processing finishes. 
+	 */
 	bool bInFlight;
+	/**
+	 * The block is in flight and has finished loaded and processing by async threads. Is true only when bInFlight is true. 
+	 * Starts false, and is set to true when and only when DoProcessing finishes with it. Cleared when block is no longer referenced.
+	 * Accessed only within FPakAsyncReadFileHandle->CriticalSection. 
+	 */
 	bool bCPUWorkIsComplete;
+	/**
+	 * True if and only if all requests touching the block canceled before the block finished processing. The block is removed from Blocks,
+	 * present in OutstandingCancelMapBlock, and still referenced as the Block pointer on the async thread.
+	 * Accessed only within FPakAsyncReadFileHandle->CriticalSection.
+	 */
 	bool bCancelledBlock;
 	FCachedAsyncBlock()
 		: RawRequest(0)
@@ -3576,6 +3621,7 @@ public:
 
 	void RequestIsComplete()
 	{
+		// Owner->CriticalSection is locked
 		if (CompleteRace.Increment() == 1)
 		{
 			check(bRequestOutstanding);
@@ -3791,14 +3837,42 @@ void FPakPrecacher::DoSignatureCheck(bool bWasCanceled, IAsyncReadRequest* Reque
 
 class FPakAsyncReadFileHandle final : public IAsyncReadFileHandle
 {
+	/** Name of the PakFile that contains the FileEntry read by this handle. Read-only after construction. */
 	FName PakFile;
+	/**
+	 * Pointer to the PakFile that contains the FileEntry read by this handle.
+	 * The pointer is read-only after construction (the PakFile exceeds the lifetime of *this).
+	 */
 	TRefCountPtr<FPakFile> ActualPakFile;
+	/** Size of the PakFile that contains the FileEntry read by this handle. Read-only after construction. */
 	int64 PakFileSize;
+	/**
+	 * Number of bytes between start of the PakFile and start of the payload (AFTER the FileEntry struct)
+	 * of the FileEntry read by this handle. Read-only after construction.
+	 */
 	int64 OffsetInPak;
+	/** Number of bytes of the payload after being uncompressed. Read-only after construction. */
 	int64 UncompressedFileSize;
+	/** PakFile's metadata about the FileEntry read by this handle. Read-only after construction. */
 	FPakEntry FileEntry;
+
+	/**
+	 * Set of Requests created by ReadRequest that will still need to access *this. Requests are removed from
+	 * LiveRequests from their destructor or when they have canceled and their blocks have finished processing.
+	 * The set is accessed only within this->CriticalSection.
+	 */
 	TSet<FPakProcessedReadRequest*> LiveRequests;
+	/**
+	 * Information about each compression block in the payload, including a refcount for how many LiveRequests
+	 * requested the block. Empty and unused if the payload is not compressed.
+	 * The array is allocated and filled with null during construction.
+	 * Pointers in the array are accessed only within this->CriticalSection. Allocated pointers are
+	 * copied by value into the async threads for loading and processing. Allocations are
+	 * cleared and reused after their request refcount goes to 0, unless they are canceled before their
+	 * processing completes. In that case they are removed from Blocks and later deleted when processing finishes.
+	 * Thread synchronization rules differ for each element of the block, see comments on struct FCachedAsyncBlock	 */
 	TArray<FCachedAsyncBlock*> Blocks;
+	/** Callback we construct to call our RawReadCallback after each block's read. Read-only after construction. */
 	FAsyncFileCallBack ReadCallbackFunction;
 	FCriticalSection CriticalSection;
 	int32 NumLiveRawRequests;
@@ -3833,7 +3907,28 @@ public:
 		UncompressedFileSize = FileEntry.UncompressedSize;
 		int64 CompressedFileSize = FileEntry.UncompressedSize;
 		CompressionMethod = InPakFile->GetInfo().GetCompressionMethod(FileEntry.CompressionMethodIndex);
-		if (CompressionMethod != NAME_None && UncompressedFileSize)
+#if !UE_BUILD_SHIPPING
+		if (GetPakCacheForcePakProcessedReads() && CompressionMethod.IsNone() && UncompressedFileSize)
+		{
+			check(FileEntry.CompressionBlocks.Num() == 0);
+			CompressionMethod = GPakFakeCompression;
+			FileEntry.CompressionBlockSize = 65536;
+			int64 EndSize = 0;
+			while (EndSize < UncompressedFileSize)
+			{
+				FPakCompressedBlock& CompressedBlock = FileEntry.CompressionBlocks.Emplace_GetRef();
+				CompressedBlock.CompressedStart = EndSize + OffsetInPak - (InPakFile->GetInfo().HasRelativeCompressedChunkOffsets() ? FileEntry.Offset : 0);
+				CompressedBlock.CompressedEnd = CompressedBlock.CompressedStart + FileEntry.CompressionBlockSize;
+				EndSize += FileEntry.CompressionBlockSize;
+				if (EndSize > UncompressedFileSize)
+				{
+					CompressedBlock.CompressedEnd -= EndSize - UncompressedFileSize;
+					EndSize = UncompressedFileSize;
+				}
+			}
+		}
+#endif
+		if (!CompressionMethod.IsNone() && UncompressedFileSize)
 		{
 			check(FileEntry.CompressionBlocks.Num());
 			CompressedFileSize = FileEntry.CompressionBlocks.Last().CompressedEnd - FileEntry.CompressionBlocks[0].CompressedStart;
@@ -3940,6 +4035,7 @@ public:
 
 	void StartBlock(int32 BlockIndex, EAsyncIOPriorityAndFlags PriorityAndFlags)
 	{
+		// this->CriticalSection is locked
 		FCachedAsyncBlock& Block = GetBlock(BlockIndex);
 		Block.bInFlight = true;
 		check(!Block.RawRequest && !Block.Processed && !Block.Raw && !Block.CPUWorkGraphEvent.GetReference() && !Block.ProcessedSize && !Block.RawSize && !Block.bCPUWorkIsComplete);
@@ -4036,9 +4132,26 @@ public:
 				check(Block.DecompressionRawSize == Block.RawSize);
 			}
 
-			bool bFailed = !FCompression::UncompressMemory(CompressionMethod, Output, Block.ProcessedSize, Block.Raw, Block.DecompressionRawSize);
-
+			bool bFailed = false;
+#if !UE_BUILD_SHIPPING
+			if (CompressionMethod != GPakFakeCompression)
+#endif
+			{
+				bFailed = !FCompression::UncompressMemory(CompressionMethod, Output, Block.ProcessedSize, Block.Raw, Block.DecompressionRawSize);
+			}
 #if !UE_BUILD_SHIPPING
+			else
+			{
+				if (bCorrupted)
+				{
+					bFailed = true;
+				}
+				else
+				{
+					check(Block.ProcessedSize == Block.DecompressionRawSize);
+					FMemory::Memcpy(Output, Block.Raw, Block.ProcessedSize);
+				}
+			}
 			if (bCorrupted && !bFailed)
 			{
 				UE_LOG(LogPakFile, Error, TEXT("The payload was corrupted, but this did not trigger a decompression failed.....pretending it failed anyway because otherwise it can crash later."));
@@ -4050,7 +4163,9 @@ public:
 			{
 				{
 					const FString HexBytes = BytesToHex(Block.Raw, FMath::Min(Block.DecompressionRawSize, 32));
-					UE_LOG(LogPakFile, Error, TEXT("Pak Decompression failed. PakFile:%s, EntryOffset:%lld, EntrySize:%lld, Method:%s, ProcessedSize:%d, RawSize:%d, Crc32:%u, BlockIndex:%d, Encrypt:%d, Delete:%d, Output:%p, Raw:%p, Processed:%p, Bytes:[%s...]"), *PakFile.ToString(), FileEntry.Offset, FileEntry.Size, *CompressionMethod.ToString(), Block.ProcessedSize, Block.DecompressionRawSize, FCrc::MemCrc32(Block.Raw, Block.DecompressionRawSize), Block.BlockIndex, FileEntry.IsEncrypted() ? 1 : 0, FileEntry.IsDeleteRecord() ? 1 : 0, Output, Block.Raw, Block.Processed, *HexBytes);
+					UE_LOG(LogPakFile, Error, TEXT("Pak Decompression failed. PakFile:%s, EntryOffset:%lld, EntrySize:%lld, Method:%s, ProcessedSize:%d, RawSize:%d, Crc32:%u, BlockIndex:%d, Encrypt:%d, Delete:%d, Output:%p, Raw:%p, Processed:%p, Bytes:[%s...]"),
+						*PakFile.ToString(), FileEntry.Offset, FileEntry.Size, *CompressionMethod.ToString(), Block.ProcessedSize, Block.DecompressionRawSize,
+						FCrc::MemCrc32(Block.Raw, Block.DecompressionRawSize), Block.BlockIndex, FileEntry.IsEncrypted() ? 1 : 0, FileEntry.IsDeleteRecord() ? 1 : 0, Output, Block.Raw, Block.Processed, *HexBytes);
 				}
 				uint8* TempBuffer = (uint8*)FMemory::Malloc(Block.RawSize);
 				{
@@ -4158,6 +4273,8 @@ public:
 	}
 	void ClearBlock(FCachedAsyncBlock& Block, bool bForDestructorShouldAlreadyBeClear = false)
 	{
+		// this->CriticalSection is locked
+
 		check(!Block.RawRequest);
 		Block.RawRequest = nullptr;
 		//check(!Block.CPUWorkGraphEvent || Block.CPUWorkGraphEvent->IsComplete());
@@ -4266,7 +4383,7 @@ public:
 
 	void GatherResults(uint8* Memory, int64 Offset, int64 BytesToRead)
 	{
-		// no lock here, I don't think it is needed because we have a ref count.
+		// CriticalSection is locked
 		int32 FirstBlock = Offset / FileEntry.CompressionBlockSize;
 		int32 LastBlock = (Offset + BytesToRead - 1) / FileEntry.CompressionBlockSize;
 		check(FirstBlock >= 0 && FirstBlock < Blocks.Num() && LastBlock >= 0 && LastBlock < Blocks.Num() && FirstBlock <= LastBlock);
@@ -4310,6 +4427,7 @@ void FPakProcessedReadRequest::CancelRawRequests()
 
 void FPakProcessedReadRequest::GatherResults()
 {
+	// Owner->CriticalSection is locked
 	if (!bUserSuppliedMemory)
 	{
 		check(!Memory);
@@ -4327,6 +4445,7 @@ void FPakProcessedReadRequest::DoneWithRawRequests()
 
 bool FPakProcessedReadRequest::CheckCompletion(const FPakEntry& FileEntry, int32 BlockIndex, TArray<FCachedAsyncBlock*>& Blocks)
 {
+	// Owner->CriticalSection is locked
 	if (!bRequestOutstanding || bHasCompleted || bHasCancelled)
 	{
 		return false;
@@ -8016,3 +8135,80 @@ void FPakPlatformFile::MakeUniquePakFilesForTheseFiles(const TArray<TArray<FStri
 }
 
 IMPLEMENT_MODULE(FPakFileModule, PakFile);
+
+#if !UE_BUILD_SHIPPING
+static void AsyncFileTest(const TArray<FString>& Args)
+{
+	FString TestFile;
+	if (Args.Num() == 0)
+	{
+		UE_LOG(LogPakFile, Error, TEXT("pak.AsyncFileTest requires a filename argument: \"pak.AsyncFileTest <filename> <size> <offset>\""));
+		return;
+	}
+
+	TestFile = Args[0];
+	int64 Size = 1;
+	if (Args.Num() > 1)
+	{
+		Size = -1;
+		LexFromString(Size, *Args[1]);
+		if (Size <= 0)
+		{
+			UE_LOG(LogPakFile, Error, TEXT("pak.AsyncFileTest size must be > 0: \"pak.AsyncFileTest <filename> <size> <offset>\""));
+			return;
+		}
+	}
+
+	int64 Offset = 0;
+	if (Args.Num() > 2)
+	{
+		Offset = -1;
+		LexFromString(Offset, *Args[2]);
+		if (Size < 0)
+		{
+			UE_LOG(LogPakFile, Error, TEXT("pak.AsyncFileTest offset must be >= 0: \"pak.AsyncFileTest <filename> <size> <offset>\""));
+			return;
+		}
+	}
+
+	IPlatformFile& PlatformFile = FPlatformFileManager::Get().GetPlatformFile();
+	TUniquePtr<IAsyncReadFileHandle> FileHandle(PlatformFile.OpenAsyncRead(*TestFile));
+	check(FileHandle);
+	{
+		TUniquePtr<IAsyncReadRequest> SizeRequest(FileHandle->SizeRequest());
+		if (!SizeRequest)
+		{
+			UE_LOG(LogPakFile, Error, TEXT("pak.AsyncFileTest: SizeRequest failed for %s."), *TestFile, Size, Offset);
+			return;
+		}
+		SizeRequest->WaitCompletion();
+		int64 TotalSize = SizeRequest->GetSizeResults();
+		SizeRequest.Reset();
+		if (Offset + Size > TotalSize)
+		{
+			UE_LOG(LogPakFile, Error, TEXT("pak.AsyncFileTest: Requested size offset is out of range for %s. Size=%" INT64_FMT", Offset=%" INT64_FMT ", End=%" INT64_FMT ", Available Size = %" INT64_FMT "."),
+				*TestFile, Size, Offset, Size + Offset, TotalSize);
+			return;
+		}
+
+		TUniquePtr<IAsyncReadRequest> ReadRequest(FileHandle->ReadRequest(Offset, Size));
+		if (!ReadRequest)
+		{
+			UE_LOG(LogPakFile, Error, TEXT("pak.AsyncFileTest: ReadRequest failed for %s size %" INT64_FMT " offset %" INT64_FMT "."), *TestFile, Size, Offset);
+			return;
+		}
+
+		ReadRequest.Reset();
+		FPlatformProcess::Sleep(3.0f);
+	}
+	FileHandle.Reset();
+
+	UE_LOG(LogPakFile, Display, TEXT("pak.AsyncFileTest: ReadRequest succeeded with no errors for %s size %" INT64_FMT " offset %" INT64_FMT "."), *TestFile, Size, Offset);
+}
+
+static FAutoConsoleCommand AsyncFileTestCmd(
+	TEXT("pak.AsyncFileTest"),
+	TEXT("Read a block of data from a file using an AsyncFileHandle. params: <filename> <size> <offset>"),
+	FConsoleCommandWithArgsDelegate::CreateStatic(AsyncFileTest));
+#endif
+
