LLM tags fixup
#rnx

diff --git a/Engine/Source/Runtime/Core/Private/IO/IoDispatcher.cpp b/Engine/Source/Runtime/Core/Private/IO/IoDispatcher.cpp
--- a/Engine/Source/Runtime/Core/Private/IO/IoDispatcher.cpp
+++ b/Engine/Source/Runtime/Core/Private/IO/IoDispatcher.cpp
@@ -51,6 +51,7 @@ public:
 
 		if (!NextFree)
 		{
+			LLM_SCOPE(ELLMTag::FileSystem);
 			//TODO: Virtual alloc
 			FBlock* Block = new FBlock;
 
diff --git a/Engine/Source/Runtime/PakFile/Private/IPlatformFilePak.cpp b/Engine/Source/Runtime/PakFile/Private/IPlatformFilePak.cpp
--- a/Engine/Source/Runtime/PakFile/Private/IPlatformFilePak.cpp
+++ b/Engine/Source/Runtime/PakFile/Private/IPlatformFilePak.cpp
@@ -7445,6 +7445,7 @@ void FPakPlatformFile::OptimizeMemoryUsageForMountedPaks()
 
 bool FPakPlatformFile::Mount(const TCHAR* InPakFilename, uint32 PakOrder, const TCHAR* InPath /*= NULL*/, bool bLoadIndex /*= true*/)
 {
+	LLM_SCOPE(ELLMTag::FileSystem);
 	bool bPakSuccess = false;
 	bool bIoStoreSuccess = true;
 	TSharedPtr<IFileHandle> PakHandle = MakeShareable(LowerLevel->OpenRead(InPakFilename));
diff --git a/Engine/Source/Runtime/RenderCore/Private/ShaderCodeLibrary.cpp b/Engine/Source/Runtime/RenderCore/Private/ShaderCodeLibrary.cpp
--- a/Engine/Source/Runtime/RenderCore/Private/ShaderCodeLibrary.cpp
+++ b/Engine/Source/Runtime/RenderCore/Private/ShaderCodeLibrary.cpp
@@ -2571,6 +2571,7 @@ public:
 
 	void PreloadPackageShaderMapsDelegate(TArrayView<const FSHAHash> ShaderMapHashes, FCoreDelegates::FAttachShaderReadRequestFunc AttachShaderReadRequestFunc)
 	{
+		LLM_SCOPE(ELLMTag::Shaders);
 		for (const FSHAHash& ShaderMapHash : ShaderMapHashes)
 		{
 			int32 ShaderMapIndex;
