Prevent problems caused by waiting on pak read requests pulling in other pak jobs while clearing old block tassk.
#rnx

diff --git a/Engine/Source/Runtime/PakFile/Private/IPlatformFilePak.cpp b/Engine/Source/Runtime/PakFile/Private/IPlatformFilePak.cpp
--- a/Engine/Source/Runtime/PakFile/Private/IPlatformFilePak.cpp
+++ b/Engine/Source/Runtime/PakFile/Private/IPlatformFilePak.cpp
@@ -2533,7 +2533,10 @@ private: // below here we assume CachedFilesScopeLock until we get to the next s
 
 			for (IAsyncReadRequest* Elem : Swapped)
 			{
-				Elem->WaitCompletion();
+				while (!Elem->PollCompletion())
+				{
+					FPlatformProcess::Sleep(0);
+				}
 				delete Elem;
 			}
 			Swapped.Empty();
