Share pak readers between threads and retire pak readers that haven't been used after a few seconds.

diff --git a/Engine/Source/Runtime/PakFile/Private/IPlatformFilePak.cpp b/Engine/Source/Runtime/PakFile/Private/IPlatformFilePak.cpp
--- a/Engine/Source/Runtime/PakFile/Private/IPlatformFilePak.cpp
+++ b/Engine/Source/Runtime/PakFile/Private/IPlatformFilePak.cpp
@@ -5000,7 +5000,7 @@ public:
 		WorkingBuffers[0] = ScratchSpace->ScratchBuffer.Get();
 		WorkingBuffers[1] = ScratchSpace->ScratchBuffer.Get() + WorkingBufferRequiredSize;
 
-		FArchive* PakReader = AcquirePakReader();
+		FSharedPakReader PakReader = AcquirePakReader();
 
 		while (Length > 0)
 		{
@@ -5176,6 +5176,47 @@ bool FPakPlatformFile::IsNonPakFilenameAllowed(const FString& InFilename)
 }
 
 
+FSharedPakReader::FSharedPakReader(FArchive* InArchive, FPakFile* InPakFile)
+	: Archive(InArchive)
+	, PakFile(InPakFile)
+{
+	check(PakFile);
+}
+
+FSharedPakReader::~FSharedPakReader()
+{
+	if(Archive)
+	{
+		PakFile->ReturnSharedReader(Archive);
+		Archive = nullptr;
+	}
+}
+
+FSharedPakReader::FSharedPakReader(FSharedPakReader&& Other)
+	: Archive(Other.Archive)
+	, PakFile(Other.PakFile)
+{
+	Other.Archive = nullptr;
+	Other.PakFile = nullptr;
+}
+
+FSharedPakReader& FSharedPakReader::operator=(FSharedPakReader&& Other)
+{
+	if (Archive)
+	{
+		PakFile->ReturnSharedReader(Archive);
+	}
+
+	Archive = Other.Archive;
+	PakFile = Other.PakFile;
+
+	Other.Archive = nullptr;
+	Other.PakFile = nullptr;
+
+	return *this;
+}
+
+
 #if IS_PROGRAM
 FPakFile::FPakFile(const TCHAR* Filename, bool bIsSigned)
 	: PakFilename(Filename)
@@ -5228,11 +5269,11 @@ FPakFile::FPakFile(IPlatformFile* LowerLevel, const TCHAR* Filename, bool bIsSig
 	, UnderlyingCacheTrimDisabled(false)
 	, bIsMounted(false)
 {
-	FArchive* Reader = GetSharedReader(LowerLevel);
+	FSharedPakReader Reader = GetSharedReader(LowerLevel);
 	if (Reader)
 	{
 		Timestamp = LowerLevel->GetTimeStamp(Filename);
-		Initialize(Reader, bLoadIndex);
+		Initialize(Reader.GetArchive(), bLoadIndex);
 	}
 }
 
@@ -5310,9 +5351,9 @@ FArchive* FPakFile::SetupSignedPakReader(FArchive* ReaderArchive, const TCHAR* F
 	return ReaderArchive;
 }
 
-void FPakFile::Initialize(FArchive* Reader, bool bLoadIndex)
+void FPakFile::Initialize(FArchive& Reader, bool bLoadIndex)
 {
-	CachedTotalSize = Reader->TotalSize();
+	CachedTotalSize = Reader.TotalSize();
 	bool bShouldLoad = false;
 	int32 CompatibleVersion = FPakInfo::PakFile_Version_Latest;
 
@@ -5330,13 +5371,13 @@ void FPakFile::Initialize(FArchive* Reader, bool bLoadIndex)
 		FileInfoPos = CachedTotalSize - Info.GetSerializedSize(CompatibleVersion);
 		if (FileInfoPos >= 0)
 		{
-			Reader->Seek(FileInfoPos);
-			Reader->Precache(FileInfoPos, 0); // Inform the archive that we're going to repeatedly serialize from the current location
+			Reader.Seek(FileInfoPos);
+			Reader.Precache(FileInfoPos, 0); // Inform the archive that we're going to repeatedly serialize from the current location
 
 			SCOPED_BOOT_TIMING("PakFile_SerilizeTrailer");
 
 			// Serialize trailer and check if everything is as expected.
-			Info.Serialize(*Reader, CompatibleVersion);
+			Info.Serialize(Reader, CompatibleVersion);
 			if (Info.Magic == FPakInfo::PakFile_Magic)
 			{
 				bShouldLoad = true;
@@ -5385,7 +5426,7 @@ void FPakFile::Initialize(FArchive* Reader, bool bLoadIndex)
 	}
 }
 
-void FPakFile::LoadIndex(FArchive* Reader)
+void FPakFile::LoadIndex(FArchive& Reader)
 {
 	if (Info.Version >= FPakInfo::PakFile_Version_PathHashIndex)
 	{
@@ -5412,7 +5453,7 @@ void FPakFile::LoadIndex(FArchive* Reader)
 	}
 }
 
-bool FPakFile::LoadIndexInternal(FArchive* Reader)
+bool FPakFile::LoadIndexInternal(FArchive& Reader)
 {
 	bHasPathHashIndex = false;
 	bHasFullDirectoryIndex = false;
@@ -5427,11 +5468,11 @@ bool FPakFile::LoadIndexInternal(FArchive* Reader)
 		return false;
 	}
 	TArray<uint8> PrimaryIndexData;
-	Reader->Seek(Info.IndexOffset);
+	Reader.Seek(Info.IndexOffset);
 	PrimaryIndexData.SetNum(Info.IndexSize);
 	{
 		SCOPED_BOOT_TIMING("PakFile_LoadPrimaryIndex");
-		Reader->Serialize(PrimaryIndexData.GetData(), Info.IndexSize);
+		Reader.Serialize(PrimaryIndexData.GetData(), Info.IndexSize);
 	}
 
 	FSHAHash ComputedHash;
@@ -5442,7 +5483,7 @@ bool FPakFile::LoadIndexInternal(FArchive* Reader)
 			UE_LOG(LogPakFile, Log, TEXT("Corrupt pak PrimaryIndex detected!"));
 			UE_LOG(LogPakFile, Log, TEXT(" Filename: %s"), *PakFilename);
 			UE_LOG(LogPakFile, Log, TEXT(" Encrypted: %d"), Info.bEncryptedIndex);
-			UE_LOG(LogPakFile, Log, TEXT(" Total Size: %d"), Reader->TotalSize());
+			UE_LOG(LogPakFile, Log, TEXT(" Total Size: %d"), Reader.TotalSize());
 			UE_LOG(LogPakFile, Log, TEXT(" Index Offset: %d"), Info.IndexOffset);
 			UE_LOG(LogPakFile, Log, TEXT(" Index Size: %d"), Info.IndexSize);
 			UE_LOG(LogPakFile, Log, TEXT(" Stored Index Hash: %s"), *Info.IndexHash.ToString());
@@ -5557,16 +5598,16 @@ bool FPakFile::LoadIndexInternal(FArchive* Reader)
 			// Should not be possible for these values (which came from the PrimaryIndex) to be invalid, since we verified the index hash of the PrimaryIndex
 			UE_LOG(LogPakFile, Log, TEXT("Corrupt pak PrimaryIndex detected!"));
 			UE_LOG(LogPakFile, Log, TEXT(" Filename: %s"), *PakFilename);
-			UE_LOG(LogPakFile, Log, TEXT(" Total Size: %d"), Reader->TotalSize());
+			UE_LOG(LogPakFile, Log, TEXT(" Total Size: %d"), Reader.TotalSize());
 			UE_LOG(LogPakFile, Log, TEXT(" PathHashIndexOffset : %d"), PathHashIndexOffset);
 			UE_LOG(LogPakFile, Log, TEXT(" PathHashIndexSize: %d"), PathHashIndexSize);
 			return false;
 		}
-		Reader->Seek(PathHashIndexOffset);
+		Reader.Seek(PathHashIndexOffset);
 		PathHashIndexData.SetNum(PathHashIndexSize);
 		{
 			SCOPED_BOOT_TIMING("PakFile_LoadPathHashIndex");
-			Reader->Serialize(PathHashIndexData.GetData(), PathHashIndexSize);
+			Reader.Serialize(PathHashIndexData.GetData(), PathHashIndexSize);
 		}
 
 		{
@@ -5576,7 +5617,7 @@ bool FPakFile::LoadIndexInternal(FArchive* Reader)
 				UE_LOG(LogPakFile, Log, TEXT("Corrupt pak PathHashIndex detected!"));
 				UE_LOG(LogPakFile, Log, TEXT(" Filename: %s"), *PakFilename);
 				UE_LOG(LogPakFile, Log, TEXT(" Encrypted: %d"), Info.bEncryptedIndex);
-				UE_LOG(LogPakFile, Log, TEXT(" Total Size: %d"), Reader->TotalSize());
+				UE_LOG(LogPakFile, Log, TEXT(" Total Size: %d"), Reader.TotalSize());
 				UE_LOG(LogPakFile, Log, TEXT(" Index Offset: %d"), FullDirectoryIndexOffset);
 				UE_LOG(LogPakFile, Log, TEXT(" Index Size: %d"), FullDirectoryIndexSize);
 				UE_LOG(LogPakFile, Log, TEXT(" Stored Index Hash: %s"), *PathHashIndexHash.ToString());
@@ -5613,17 +5654,17 @@ bool FPakFile::LoadIndexInternal(FArchive* Reader)
 			// Should not be possible for these values (which came from the PrimaryIndex) to be invalid, since we verified the index hash of the PrimaryIndex
 			UE_LOG(LogPakFile, Log, TEXT("Corrupt pak PrimaryIndex detected!"));
 			UE_LOG(LogPakFile, Log, TEXT(" Filename: %s"), *PakFilename);
-			UE_LOG(LogPakFile, Log, TEXT(" Total Size: %d"), Reader->TotalSize());
+			UE_LOG(LogPakFile, Log, TEXT(" Total Size: %d"), Reader.TotalSize());
 			UE_LOG(LogPakFile, Log, TEXT(" FullDirectoryIndexOffset : %d"), FullDirectoryIndexOffset );
 			UE_LOG(LogPakFile, Log, TEXT(" FullDirectoryIndexSize: %d"), FullDirectoryIndexSize);
 			return false;
 		}
 		TArray<uint8> FullDirectoryIndexData;
-		Reader->Seek(FullDirectoryIndexOffset );
+		Reader.Seek(FullDirectoryIndexOffset );
 		FullDirectoryIndexData.SetNum(FullDirectoryIndexSize);
 		{
 			SCOPED_BOOT_TIMING("PakFile_LoadDirectoryIndex");
-			Reader->Serialize(FullDirectoryIndexData.GetData(), FullDirectoryIndexSize);
+			Reader.Serialize(FullDirectoryIndexData.GetData(), FullDirectoryIndexSize);
 		}
 
 		{
@@ -5633,7 +5674,7 @@ bool FPakFile::LoadIndexInternal(FArchive* Reader)
 				UE_LOG(LogPakFile, Log, TEXT("Corrupt pak FullDirectoryIndex detected!"));
 				UE_LOG(LogPakFile, Log, TEXT(" Filename: %s"), *PakFilename);
 				UE_LOG(LogPakFile, Log, TEXT(" Encrypted: %d"), Info.bEncryptedIndex);
-				UE_LOG(LogPakFile, Log, TEXT(" Total Size: %d"), Reader->TotalSize());
+				UE_LOG(LogPakFile, Log, TEXT(" Total Size: %d"), Reader.TotalSize());
 				UE_LOG(LogPakFile, Log, TEXT(" Index Offset: %d"), FullDirectoryIndexOffset);
 				UE_LOG(LogPakFile, Log, TEXT(" Index Size: %d"), FullDirectoryIndexSize);
 				UE_LOG(LogPakFile, Log, TEXT(" Stored Index Hash: %s"), *FullDirectoryIndexHash.ToString());
@@ -5676,7 +5717,7 @@ bool FPakFile::LoadIndexInternal(FArchive* Reader)
 	return true;
 }
 
-bool FPakFile::LoadLegacyIndex(FArchive* Reader)
+bool FPakFile::LoadLegacyIndex(FArchive& Reader)
 {
 	bHasPathHashIndex = false;
 	bHasFullDirectoryIndex = false;
@@ -5694,9 +5735,9 @@ bool FPakFile::LoadLegacyIndex(FArchive* Reader)
 		return false;
 	}
 
-	Reader->Seek(Info.IndexOffset);
+	Reader.Seek(Info.IndexOffset);
 	IndexData.SetNum(Info.IndexSize);
-	Reader->Serialize(IndexData.GetData(), Info.IndexSize);
+	Reader.Serialize(IndexData.GetData(), Info.IndexSize);
 
 	FSHAHash ComputedHash;
 	if (!DecryptAndValidateIndex(Reader, IndexData, Info.IndexHash, ComputedHash))
@@ -5704,7 +5745,7 @@ bool FPakFile::LoadLegacyIndex(FArchive* Reader)
 		UE_LOG(LogPakFile, Log, TEXT("Corrupt pak index detected!"));
 		UE_LOG(LogPakFile, Log, TEXT(" Filename: %s"), *PakFilename);
 		UE_LOG(LogPakFile, Log, TEXT(" Encrypted: %d"), Info.bEncryptedIndex);
-		UE_LOG(LogPakFile, Log, TEXT(" Total Size: %d"), Reader->TotalSize());
+		UE_LOG(LogPakFile, Log, TEXT(" Total Size: %d"), Reader.TotalSize());
 		UE_LOG(LogPakFile, Log, TEXT(" Index Offset: %d"), Info.IndexOffset);
 		UE_LOG(LogPakFile, Log, TEXT(" Index Size: %d"), Info.IndexSize);
 		UE_LOG(LogPakFile, Log, TEXT(" Stored Index Hash: %s"), *Info.IndexHash.ToString());
@@ -5769,7 +5810,7 @@ bool FPakFile::LoadLegacyIndex(FArchive* Reader)
 	return true;
 }
 
-bool FPakFile::DecryptAndValidateIndex(FArchive* Reader, TArray<uint8>& IndexData, FSHAHash& InExpectedHash, FSHAHash& OutActualHash)
+bool FPakFile::DecryptAndValidateIndex(FArchive& Reader, TArray<uint8>& IndexData, FSHAHash& InExpectedHash, FSHAHash& OutActualHash)
 {
 	// Decrypt if necessary
 	if (Info.bEncryptedIndex)
@@ -6554,7 +6595,7 @@ bool FPakFile::Check()
 	UE_LOG(LogPakFile, Display, TEXT("Checking pak file \"%s\". This may take a while..."), *PakFilename);
 	double StartTime = FPlatformTime::Seconds();
 
-	FArchive& PakReader = *GetSharedReader(NULL);
+	FSharedPakReader PakReader = GetSharedReader(nullptr);
 	int32 ErrorCount = 0;
 	int32 FileCount = 0;
 
@@ -6577,12 +6618,12 @@ bool FPakFile::Check()
 		}
 
 		int64 CurrentPos = 0;
-		const int64 Size = PakReader.TotalSize();
+		const int64 Size = PakReader->TotalSize();
 		while (CurrentPos < Size)
 		{
-			PakReader.Seek(CurrentPos);
+			PakReader->Seek(CurrentPos);
 			uint8 Byte = 0;
-			PakReader << Byte;
+			PakReader.GetArchive() << Byte;
 			CurrentPos += FPakInfo::MaxChunkDataSize;
 		}
 
@@ -6619,16 +6660,16 @@ bool FPakFile::Check()
 			}
 
 			void* FileContents = FMemory::Malloc(EntryFromIndex.Size);
-			PakReader.Seek(EntryFromIndex.Offset);
+			PakReader->Seek(EntryFromIndex.Offset);
 			uint32 SerializedCrcTest = 0;
 			FPakEntry EntryFromPayload;
-			EntryFromPayload.Serialize(PakReader, GetInfo().Version);
+			EntryFromPayload.Serialize(PakReader.GetArchive(), GetInfo().Version);
 			if (!EntryFromPayload.IndexDataEquals(EntryFromIndex))
 			{
 				UE_LOG(LogPakFile, Error, TEXT("Index FPakEntry does not match Payload FPakEntry for %s."), GetEntryName(It));
 				ErrorCount++;
 			}
-			PakReader.Serialize(FileContents, EntryFromIndex.Size);
+			PakReader->Serialize(FileContents, EntryFromIndex.Size);
 
 			uint8 TestHash[20];
 			FSHA1::HashBuffer(FileContents, EntryFromIndex.Size, TestHash);
@@ -6815,14 +6856,12 @@ bool FPakFile::RecreatePakReaders(IPlatformFile* LowerLevel)
 	// need to reset the decryptor as it will hold a pointer to the first created pak reader
 	Decryptor.Reset();
 
-	TMap<uint32, TUniquePtr<FArchive>> TempReaderMap;
+	TArray<FArchiveAndLastAccessTime> TempReaders;
 
-	// Create a new PakReader *per* thread that was already mapped
-	for (const TPair<uint32, TUniquePtr<FArchive>>& Reader : ReaderMap)
+	// Create a new PakReader *per* instance that was already mapped
+	for (const FArchiveAndLastAccessTime& Reader : Readers)
 	{
 		FArchive* PakReader = nullptr;
-		uint32 Thread = Reader.Key;
-
 		if (LowerLevel != nullptr)
 		{
 			IFileHandle* PakHandle = LowerLevel->OpenRead(*GetFilename());
@@ -6841,34 +6880,29 @@ bool FPakFile::RecreatePakReaders(IPlatformFile* LowerLevel)
 			UE_LOG(LogPakFile, Warning, TEXT("Unable to re-create pak \"%s\" handle"), *GetFilename());
 			return false;
 		}
-
-#if DO_CHECK
-		FArchive* Proxy = new FThreadCheckingArchiveProxy(PakReader, Thread);
-		TempReaderMap.Emplace(Thread, Proxy);
-#else
-		TempReaderMap.Emplace(Thread, PakReader);
-#endif //DO_CHECK
+		TempReaders.Add(FArchiveAndLastAccessTime{ TUniquePtr<FArchive>{ PakReader }, Reader.LastAccessTime });
 	}
 
-	// replace the current ReaderMap with the newly created pak readers leaving them to out of scope
-	ReaderMap = MoveTemp(TempReaderMap);
+	// replace the current Readers with the newly created pak readers leaving them to out of scope
+	Readers= MoveTemp(TempReaders);
 
 	return true;
 }
 
-FArchive* FPakFile::GetSharedReader(IPlatformFile* LowerLevel)
+LLM_DEFINE_TAG(PakSharedReaders);
+
+FSharedPakReader FPakFile::GetSharedReader(IPlatformFile* LowerLevel)
 {
-	uint32 Thread = FPlatformTLS::GetCurrentThreadId();
+	LLM_SCOPE_BYTAG(PakSharedReaders);
 	FArchive* PakReader = nullptr;
 	{
 		FScopeLock ScopedLock(&CriticalSection);
-		TUniquePtr<FArchive>* ExistingReader = ReaderMap.Find(Thread);
-		if (ExistingReader)
+		if (Readers.Num())
 		{
-			PakReader = ExistingReader->Get();
+			FArchiveAndLastAccessTime Reader = Readers.Pop();
+			PakReader = Reader.Archive.Release();
 		}
-
-		if (!PakReader)
+		else
 		{
 			// Create a new FArchive reader and pass it to the new handle.
 			if (LowerLevel != nullptr)
@@ -6883,22 +6917,37 @@ FArchive* FPakFile::GetSharedReader(IPlatformFile* LowerLevel)
 			{
 				PakReader = CreatePakReader(*GetFilename());
 			}
+
 			if (!PakReader)
 			{
 				UE_LOG(LogPakFile, Warning, TEXT("Unable to create pak \"%s\" handle"), *GetFilename());
-				return nullptr;
 			}
+		}
+	}
 
-#if DO_CHECK
-			FArchive* Proxy = new FThreadCheckingArchiveProxy(PakReader, Thread);
-			ReaderMap.Emplace(Thread, Proxy);
-			PakReader = Proxy;
-#else //DO_CHECK
-			ReaderMap.Emplace(Thread, PakReader);
-#endif //DO_CHECK
+	return FSharedPakReader(PakReader, this);
+}
+
+void FPakFile::ReturnSharedReader(FArchive* Archive)
+{
+	FScopeLock ScopedLock(&CriticalSection);
+	Readers.Push(FArchiveAndLastAccessTime{ TUniquePtr<FArchive>{Archive }, FPlatformTime::Seconds()});
+}
+
+void FPakFile::ReleaseOldReaders(double MaxAgeSeconds)
+{
+	FScopeLock ScopedLock(&CriticalSection);
+	double SearchTime = FPlatformTime::Seconds() - MaxAgeSeconds;
+	for (int32 i = Readers.Num() - 1; i >= 0; --i)
+	{
+		const FArchiveAndLastAccessTime& Reader = Readers[i];
+		if (Reader.LastAccessTime <= SearchTime)
+		{
+			// Remove this and all readers older than it (pushed before it)
+			Readers.RemoveAt(0, i + 1);
+			break;
 		}
 	}
-	return PakReader;
 }
 
 const FPakEntryLocation* FPakFile::FindLocationFromIndex(const FString& FullPath, const FString& MountPoint, const FPathHashIndex& PathHashIndex, uint64 PathHashSeed, int32 PakFileVersion)
@@ -7087,6 +7136,8 @@ FPakPlatformFile::FPakPlatformFile()
 
 FPakPlatformFile::~FPakPlatformFile()
 {
+	FTSTicker::GetCoreTicker().RemoveTicker(RetireReadersHandle);
+
 	FCoreDelegates::GetRegisterEncryptionKeyMulticastDelegate().RemoveAll(this);
 	FCoreDelegates::OnFEngineLoopInitComplete.RemoveAll(this);
 
@@ -7241,6 +7292,8 @@ bool FPakPlatformFile::Initialize(IPlatformFile* Inner, const TCHAR* CmdLine)
 	check(Inner != NULL);
 	LowerLevel = Inner;
 
+	RetireReadersHandle = FTSTicker::GetCoreTicker().AddTicker(TEXT("RetirePakReaders"), 1.0f, [this](float){ this->ReleaseOldReaders(); return true; });
+
 #if EXCLUDE_NONPAK_UE_EXTENSIONS && !WITH_EDITOR
 	// Extensions for file types that should only ever be in a pak file. Used to stop unnecessary access to the lower level platform file
 	ExcludedNonPakExtensions.Add(TEXT("uasset"));
@@ -7363,6 +7416,27 @@ int32 GetRecursiveAllocatedSize(const FPakFile::FDirectoryIndex& Index)
 }
 #endif
 
+static float GPakReaderReleaseDelay = 5.0f;
+static FAutoConsoleVariableRef CVarPakReaderReleaseDelay(
+	TEXT("pak.ReaderReleaseDelay"),
+	GPakReaderReleaseDelay,
+	TEXT("If > 0, then synchronous pak readers older than this will be deleted.")
+);
+
+void FPakPlatformFile::ReleaseOldReaders()
+{
+	if (GPakReaderReleaseDelay == 0.0f) {
+		return;
+	}
+
+	FScopeLock ScopedLock(&PakListCritical);
+
+	for (FPakListEntry& Entry : PakFiles)
+	{
+		Entry.PakFile->ReleaseOldReaders(GPakReaderReleaseDelay);
+	}
+}
+
 void FPakPlatformFile::OptimizeMemoryUsageForMountedPaks()
 {
 #if !UE_BUILD_SHIPPING
@@ -7606,7 +7680,7 @@ bool FPakPlatformFile::Unmount(const TCHAR* InPakFilename)
 
 		if (UnmountedPak)
 		{
-			UnmountedPak->ReaderMap.Empty();
+			UnmountedPak->Readers.Empty();
 		}
 	}
 #if USE_PAK_PRECACHE
@@ -7645,9 +7719,8 @@ bool FPakPlatformFile::ReloadPakReaders()
 
 IFileHandle* FPakPlatformFile::CreatePakFileHandle(const TCHAR* Filename, const TRefCountPtr<FPakFile>& PakFile, const FPakEntry* FileEntry)
 {
-	IFileHandle* Result = NULL;
-	bool bNeedsDelete = true;
-	TFunction<FArchive*()> AcquirePakReader = [StoredPakFile=TRefCountPtr<FPakFile>(PakFile), LowerLevelPlatformFile = LowerLevel]()
+	IFileHandle* Result = nullptr;
+	TAcquirePakReaderFunction AcquirePakReader = [StoredPakFile=TRefCountPtr<FPakFile>(PakFile), LowerLevelPlatformFile = LowerLevel]() -> FSharedPakReader
 	{
 		return StoredPakFile->GetSharedReader(LowerLevelPlatformFile);
 	};
@@ -7658,20 +7731,20 @@ IFileHandle* FPakPlatformFile::CreatePakFileHandle(const TCHAR* Filename, const
 	{
 		if (FileEntry->IsEncrypted())
 		{
-			Result = new FPakFileHandle< FPakCompressedReaderPolicy<FPakSimpleEncryption> >(ConstPakFile, *FileEntry, AcquirePakReader, bNeedsDelete);
+			Result = new FPakFileHandle< FPakCompressedReaderPolicy<FPakSimpleEncryption> >(ConstPakFile, *FileEntry, AcquirePakReader);
 		}
 		else
 		{
-			Result = new FPakFileHandle< FPakCompressedReaderPolicy<> >(ConstPakFile, *FileEntry, AcquirePakReader, bNeedsDelete);
+			Result = new FPakFileHandle< FPakCompressedReaderPolicy<> >(ConstPakFile, *FileEntry, AcquirePakReader);
 		}
 	}
 	else if (FileEntry->IsEncrypted())
 	{
-		Result = new FPakFileHandle< FPakReaderPolicy<FPakSimpleEncryption> >(ConstPakFile, *FileEntry, AcquirePakReader, bNeedsDelete);
+		Result = new FPakFileHandle< FPakReaderPolicy<FPakSimpleEncryption> >(ConstPakFile, *FileEntry, AcquirePakReader);
 	}
 	else
 	{
-		Result = new FPakFileHandle<>(ConstPakFile, *FileEntry, AcquirePakReader, bNeedsDelete);
+		Result = new FPakFileHandle<>(ConstPakFile, *FileEntry, AcquirePakReader);
 	}
 
 	return Result;
diff --git a/Engine/Source/Runtime/PakFile/Public/IPlatformFilePak.h b/Engine/Source/Runtime/PakFile/Public/IPlatformFilePak.h
--- a/Engine/Source/Runtime/PakFile/Public/IPlatformFilePak.h
+++ b/Engine/Source/Runtime/PakFile/Public/IPlatformFilePak.h
@@ -16,6 +16,7 @@
 #include "GenericPlatform/GenericPlatformChunkInstall.h"
 #include "Serialization/MemoryImage.h"
 #include "Templates/RefCounting.h"
+#include "Containers/Ticker.h"
 
 class FChunkCacheWorker;
 class IAsyncReadFileHandle;
@@ -625,6 +626,35 @@ FORCEINLINE FArchive& operator<<(FArchive& Ar, FPakEntryLocation& PakEntryLocati
 	return Ar;
 }
 
+// Wrapper for a pointer to a shared pak reader archive that has been temporarily acquired. 
+class PAKFILE_API FSharedPakReader final 
+{
+	friend class FPakFile;
+
+	FArchive* Archive = nullptr;
+	FPakFile* PakFile = nullptr; // Pak file to return ownership to on destruction
+
+	FSharedPakReader(FArchive* InArchive, FPakFile* InPakFile);
+
+public:
+	~FSharedPakReader();
+
+	FSharedPakReader(const FSharedPakReader& Other) = delete;
+	FSharedPakReader& operator=(const FSharedPakReader& Other) = delete;
+	FSharedPakReader(FSharedPakReader&& Other);
+	FSharedPakReader& operator=(FSharedPakReader&& Other);
+
+	explicit operator bool() const { return Archive != nullptr; }
+	bool operator==(nullptr_t) { return Archive == nullptr; }
+	bool operator!=(nullptr_t) { return Archive != nullptr; }
+	FArchive* operator->() { return Archive; }
+
+	
+	// USE WITH CARE, the FSharedPakReader must live longer than this reference to prevent the archive being used by another thread. Do not call on a temporary return value!
+	FArchive& GetArchive() { return *Archive; } 
+
+};
+
 /** Pak directory type mapping a filename to an FPakEntryLocation. */
 typedef TMap<FString, FPakEntryLocation> FPakDirectory;
 
@@ -685,6 +715,12 @@ public:
 	/** Recreates the pak reader for each thread */
 	bool RecreatePakReaders(IPlatformFile* LowerLevel);
 
+	struct FArchiveAndLastAccessTime 
+	{
+		TUniquePtr<FArchive> Archive;
+		double LastAccessTime;
+	};
+
 private:
 	friend class FPakPlatformFile;
 
@@ -693,9 +729,9 @@ private:
 	FName PakFilenameName;
 	/** Archive to serialize the pak file from. */
 	TUniquePtr<class FChunkCacheWorker> Decryptor;
-	/** Map of readers assigned to threads. */
-	TMap<uint32, TUniquePtr<FArchive>> ReaderMap;
-	/** Critical section for accessing ReaderMap. */
+	/** List of readers and when they were last used. */
+	TArray<FArchiveAndLastAccessTime> Readers;
+	/** Critical section for accessing Readers. */
 	FCriticalSection CriticalSection;
 	/** Pak file info (trailer). */
 	FPakInfo Info;
@@ -927,7 +963,17 @@ public:
 	 *
 	 * @return Pointer to pak file archive used to read data from pak.
 	 */
-	FArchive* GetSharedReader(IPlatformFile* LowerLevel);
+#if IS_PROGRAM
+	FSharedPakReader GetSharedReader();
+#else
+	FSharedPakReader GetSharedReader(IPlatformFile* LowerLevel);
+#endif
+
+	// Return a shared pak reader. Should only be called from the FSharedPakReader's destructor.
+	void ReturnSharedReader(FArchive* SharedReader);
+
+	// Delete all readers that haven't been used in MaxAgeSeconds.
+	void ReleaseOldReaders(double MaxAgeSeconds);
 
 	/**
 	 * Finds an entry in the pak file matching the given filename.
@@ -1357,7 +1403,7 @@ public:
 		}
 		else
 		{
-			FArchive* Reader = GetSharedReader(nullptr);
+			TUniquePtr<FArchive> Reader {CreatePakReader(*GetFilename())};
 			Reader->Seek(PakEntry.Offset);
 			FPakEntry SerializedEntry;
 			SerializedEntry.Serialize(*Reader, GetInfo().Version);
@@ -1555,12 +1601,12 @@ private:
 	/**
 	 * Initializes the pak file.
 	 */
-	void Initialize(FArchive* Reader, bool bLoadIndex = true);
+	void Initialize(FArchive& Reader, bool bLoadIndex = true);
 
 	/**
 	 * Loads and initializes pak file index.
 	 */
-	void LoadIndex(FArchive* Reader);
+	void LoadIndex(FArchive& Reader);
 
 	/**
 	  * Returns the FPakEntry pointed to by the given FPakEntryLocation, forwards to the static GetPakEntry with data from *this
@@ -1622,13 +1668,13 @@ private:
 	static void DecodePakEntry(const uint8* SourcePtr, FPakEntry& OutEntry, const FPakInfo& InInfo);
 
 	/* Internal index loading function that returns false if index loading fails due to an intermittent IO error. Allows LoadIndex to retry or throw a fatal as required */
-	bool LoadIndexInternal(FArchive* Reader);
+	bool LoadIndexInternal(FArchive& Reader);
 
 	/* Legacy index loading function for PakFiles saved before FPakInfo::PakFile_Version_PathHashIndex */
-	bool LoadLegacyIndex(FArchive* Reader);
+	bool LoadLegacyIndex(FArchive& Reader);
 
 	/* Helper function for LoadIndexInternal; each array of Index bytes read from the file needs to be independently decrypted and checked for corruption */
-	bool DecryptAndValidateIndex(FArchive* Reader, TArray<uint8>& IndexData, FSHAHash& InExpectedHash, FSHAHash& OutActualHash);
+	bool DecryptAndValidateIndex(FArchive& Reader, TArray<uint8>& IndexData, FSHAHash& InExpectedHash, FSHAHash& OutActualHash);
 
 	/* Manually add a file to a pak file */
 	void AddSpecialFile(const FPakEntry& Entry, const FString& Filename);
@@ -1728,7 +1774,7 @@ public:
 /**
  * Typedef for a function that returns an archive to use for accessing an underlying pak file
  */
-typedef TFunction<FArchive*()> TAcquirePakReaderFunction;
+typedef TFunction<FSharedPakReader()> TAcquirePakReaderFunction;
 
 template< typename EncryptionPolicy = FPakNoEncryption >
 class PAKFILE_API FPakReaderPolicy
@@ -1762,7 +1808,7 @@ public:
 		const constexpr int64 Alignment = (int64)EncryptionPolicy::Alignment;
 		const constexpr int64 AlignmentMask = ~(Alignment - 1);
 		uint8 TempBuffer[Alignment];
-		FArchive* PakReader = AcquirePakReader();
+		FSharedPakReader PakReader = AcquirePakReader();
 		if (EncryptionPolicy::AlignReadRequest(DesiredPosition) != DesiredPosition)
 		{
 			int64 Start = DesiredPosition & AlignmentMask;
@@ -1802,9 +1848,7 @@ public:
  */
 template< typename ReaderPolicy = FPakReaderPolicy<> >
 class PAKFILE_API FPakFileHandle : public IFileHandle
-{	
-	/** True if PakReader is shared and should not be deleted by this handle. */
-	const bool bSharedReader;
+{
 	/** Current read position. */
 	int64 ReadPos;
 	/** Class that controls reading from pak file */
@@ -1815,8 +1859,8 @@ class PAKFILE_API FPakFileHandle : public IFileHandle
 public:
 
 	UE_DEPRECATED(4.27, "Use constructor that takes a TRefCountPtr<FPakFile> instead")
-	FPakFileHandle(const FPakFile& InPakFile, const FPakEntry& InPakEntry, TAcquirePakReaderFunction& InAcquirePakReaderFunction, bool bIsSharedReader)
-		: FPakFileHandle(TRefCountPtr<const FPakFile>(&InPakFile), InPakEntry, InAcquirePakReaderFunction, bIsSharedReader)
+	FPakFileHandle(const FPakFile& InPakFile, const FPakEntry& InPakEntry, TAcquirePakReaderFunction& InAcquirePakReaderFunction)
+		: FPakFileHandle(TRefCountPtr<const FPakFile>(&InPakFile), InPakEntry, InAcquirePakReaderFunction)
 	{
 	}
 
@@ -1827,9 +1871,8 @@ public:
 	 * @param InPakEntry Entry in the pak file.
 	 * @param InAcquirePakReaderFunction Function that returns the archive to use for serialization. The result of this should not be cached, but reacquired on each serialization operation
 	 */
-	FPakFileHandle(const TRefCountPtr<const FPakFile>& InPakFile, const FPakEntry& InPakEntry, TAcquirePakReaderFunction& InAcquirePakReaderFunction, bool bIsSharedReader)
-		: bSharedReader(bIsSharedReader)
-		, ReadPos(0)
+	FPakFileHandle(const TRefCountPtr<const FPakFile>& InPakFile, const FPakEntry& InPakEntry, TAcquirePakReaderFunction& InAcquirePakReaderFunction)
+		: ReadPos(0)
 		, Reader(*InPakFile, InPakEntry, InAcquirePakReaderFunction)
 		, PakFile(InPakFile)
 	{
@@ -1837,8 +1880,8 @@ public:
 	}
 
 	UE_DEPRECATED(4.27, "Use constructor that takes a TRefCountPtr<FPakFile> instead")
-	FPakFileHandle(const FPakFile& InPakFile, const FPakEntry& InPakEntry, FArchive* InPakReader, bool bIsSharedReader)
-		: FPakFileHandle(TRefCountPtr<const FPakFile>(&InPakFile), InPakEntry, InPakReader, bIsSharedReader)
+	FPakFileHandle(const FPakFile& InPakFile, const FPakEntry& InPakEntry, FArchive* InPakReader)
+		: FPakFileHandle(TRefCountPtr<const FPakFile>(&InPakFile), InPakEntry, InPakReader)
 	{
 	}
 
@@ -1849,10 +1892,9 @@ public:
 	 * @param InPakEntry Entry in the pak file.
 	 * @param InPakFile Pak file.
 	 */
-	FPakFileHandle(const TRefCountPtr<const FPakFile>& InPakFile, const FPakEntry& InPakEntry, FArchive* InPakReader, bool bIsSharedReader)
-		: bSharedReader(bIsSharedReader)
-		, ReadPos(0)
-		, Reader(*InPakFile, InPakEntry, [InPakReader]() { return InPakReader; })
+	FPakFileHandle(const TRefCountPtr<const FPakFile>& InPakFile, const FPakEntry& InPakEntry, FArchive* InPakReader)
+		: ReadPos(0)
+		, Reader(*InPakFile, InPakEntry, InPakReader)
 		, PakFile(InPakFile)
 	{
 		INC_DWORD_STAT(STAT_PakFile_NumOpenHandles);
@@ -1863,11 +1905,6 @@ public:
 	 */
 	virtual ~FPakFileHandle()
 	{
-		if (!bSharedReader)
-		{
-			delete Reader.AcquirePakReader();
-		}
-
 		DEC_DWORD_STAT(STAT_PakFile_NumOpenHandles);
 	}
 
@@ -1897,9 +1934,9 @@ public:
 		if (!Reader.PakEntry.Verified)
 		{
 			FPakEntry FileHeader;
-			FArchive* PakReader = Reader.AcquirePakReader();
+			FSharedPakReader PakReader = Reader.AcquirePakReader();
 			PakReader->Seek(Reader.PakEntry.Offset);
-			FileHeader.Serialize(*PakReader, Reader.PakFile.GetInfo().Version);
+			FileHeader.Serialize(PakReader.GetArchive(), Reader.PakFile.GetInfo().Version);
 			if (FPakEntry::VerifyPakEntriesMatch(Reader.PakEntry, FileHeader))
 			{
 				Reader.PakEntry.Verified = true;
@@ -1993,6 +2030,8 @@ class PAKFILE_API FPakPlatformFile : public IPlatformFile
 	FString GameUserSettingsIniFilename;
 	TSharedPtr<IIoDispatcherFileBackend> IoDispatcherFileBackend;
 
+	FTSTicker::FDelegateHandle RetireReadersHandle;
+
 	/**
 	 * Gets mounted pak files
 	 */
@@ -3052,6 +3091,8 @@ public:
 	/** Returns the RelativePathFromMount Filename for every Filename found in the Iostore Container that relates to the provided block indexes */
 	static void GetFilenamesFromIostoreByBlockIndex(const FString& InContainerName, const TArray<int32>& InBlockIndex, TArray<FString>& OutFileList);
 
+	void ReleaseOldReaders();
+
 	// BEGIN Console commands
 #if !UE_BUILD_SHIPPING
 	void HandlePakListCommand(const TCHAR* Cmd, FOutputDevice& Ar);
