Fix editor and clang compiles for pak file changes.

diff --git a/Engine/Source/Runtime/PakFile/Private/IPlatformFilePak.cpp b/Engine/Source/Runtime/PakFile/Private/IPlatformFilePak.cpp
--- a/Engine/Source/Runtime/PakFile/Private/IPlatformFilePak.cpp
+++ b/Engine/Source/Runtime/PakFile/Private/IPlatformFilePak.cpp
@@ -5296,7 +5296,7 @@ FPakFile::FPakFile(FArchive* Archive)
 	, UnderlyingCacheTrimDisabled(false)
 	, bIsMounted(false)
 {
-	Initialize(Archive);
+	Initialize(*Archive);
 }
 #endif
 
diff --git a/Engine/Source/Runtime/PakFile/Public/IPlatformFilePak.h b/Engine/Source/Runtime/PakFile/Public/IPlatformFilePak.h
--- a/Engine/Source/Runtime/PakFile/Public/IPlatformFilePak.h
+++ b/Engine/Source/Runtime/PakFile/Public/IPlatformFilePak.h
@@ -626,6 +626,9 @@ FORCEINLINE FArchive& operator<<(FArchive& Ar, FPakEntryLocation& PakEntryLocati
 	return Ar;
 }
 
+
+class FPakFile;
+
 // Wrapper for a pointer to a shared pak reader archive that has been temporarily acquired. 
 class PAKFILE_API FSharedPakReader final 
 {
