More UnrealPak compile fixes.

diff --git a/Engine/Source/Runtime/PakFile/Private/IPlatformFilePak.cpp b/Engine/Source/Runtime/PakFile/Private/IPlatformFilePak.cpp
--- a/Engine/Source/Runtime/PakFile/Private/IPlatformFilePak.cpp
+++ b/Engine/Source/Runtime/PakFile/Private/IPlatformFilePak.cpp
@@ -5239,11 +5239,11 @@ FPakFile::FPakFile(const TCHAR* Filename, bool bIsSigned)
 	, UnderlyingCacheTrimDisabled(false)
 	, bIsMounted(false)
 {
-	FArchive* Reader = GetSharedReader(NULL);
+	FSharedPakReader Reader = GetSharedReader(NULL);
 	if (Reader)
 	{
 		Timestamp = IFileManager::Get().GetTimeStamp(Filename);
-		Initialize(Reader);
+		Initialize(Reader.GetArchive());
 	}
 }
 #endif
diff --git a/Engine/Source/Runtime/PakFile/Public/IPlatformFilePak.h b/Engine/Source/Runtime/PakFile/Public/IPlatformFilePak.h
--- a/Engine/Source/Runtime/PakFile/Public/IPlatformFilePak.h
+++ b/Engine/Source/Runtime/PakFile/Public/IPlatformFilePak.h
@@ -966,11 +966,7 @@ public:
 	 *
 	 * @return Pointer to pak file archive used to read data from pak.
 	 */
-#if IS_PROGRAM
-	FSharedPakReader GetSharedReader();
-#else
 	FSharedPakReader GetSharedReader(IPlatformFile* LowerLevel);
-#endif
 
 	// Return a shared pak reader. Should only be called from the FSharedPakReader's destructor.
 	void ReturnSharedReader(FArchive* SharedReader);
