Don't initialize the I/O dispatcher from PAK when running legacy cook-on-th-fly (-FileHostIP)

#rnx

diff --git a/Engine/Source/Runtime/PakFile/Private/IPlatformFilePak.cpp b/Engine/Source/Runtime/PakFile/Private/IPlatformFilePak.cpp
--- a/Engine/Source/Runtime/PakFile/Private/IPlatformFilePak.cpp
+++ b/Engine/Source/Runtime/PakFile/Private/IPlatformFilePak.cpp
@@ -7318,7 +7318,7 @@ bool FPakPlatformFile::Initialize(IPlatformFile* Inner, const TCHAR* CmdLine)
 #endif
 
 	FString GlobalUTocPath = FString::Printf(TEXT("%sPaks/global.utoc"), *FPaths::ProjectContentDir());
-	if (FPlatformFileManager::Get().GetPlatformFile().FileExists(*GlobalUTocPath))
+	if (!IsRunningCookOnTheFly() && FPlatformFileManager::Get().GetPlatformFile().FileExists(*GlobalUTocPath))
 	{
 		FIoStatus IoDispatcherInitStatus = FIoDispatcher::Initialize();
 		if (IoDispatcherInitStatus.IsOk())
