Fix crash in pak decryptor by giving the decryptor its own archive.

diff --git a/Engine/Source/Runtime/PakFile/Private/IPlatformFilePak.cpp b/Engine/Source/Runtime/PakFile/Private/IPlatformFilePak.cpp
--- a/Engine/Source/Runtime/PakFile/Private/IPlatformFilePak.cpp
+++ b/Engine/Source/Runtime/PakFile/Private/IPlatformFilePak.cpp
@@ -5310,20 +5310,28 @@ bool FPakFile::PassedSignatureChecks() const
 	return Decryptor.IsValid() && Decryptor->IsValid();
 }
 
-FArchive* FPakFile::CreatePakReader(const TCHAR* Filename)
-{
-	FArchive* ReaderArchive = IFileManager::Get().CreateFileReader(Filename);
-	return SetupSignedPakReader(ReaderArchive, Filename);
-}
 
-FArchive* FPakFile::CreatePakReader(IFileHandle& InHandle, const TCHAR* Filename)
+FArchive* FPakFile::CreatePakReader(IPlatformFile* LowerLevel, const TCHAR* Filename)
 {
-	FArchive* ReaderArchive = new FArchiveFileReaderGeneric(&InHandle, Filename, InHandle.Size());
-	return SetupSignedPakReader(ReaderArchive, Filename);
-}
+	auto MakeArchive = [&]() -> FArchive* { 
+		if (LowerLevel)
+		{
+			if( IFileHandle* Handle = LowerLevel->OpenRead(Filename) )
+			{
+				return new FArchiveFileReaderGeneric(Handle, Filename, Handle->Size());
+			}
+			else
+			{
+				return nullptr;
+			}
+		}
+		else
+		{
+			return IFileManager::Get().CreateFileReader(Filename); 
+		}
+	};
 
-FArchive* FPakFile::SetupSignedPakReader(FArchive* ReaderArchive, const TCHAR* Filename)
-{
+	bool bNeedsDecryptor = false;
 	if (FPlatformProperties::RequiresCookedData())
 	{
 		bool bShouldCheckSignature = bSigned || FParse::Param(FCommandLine::Get(), TEXT("signedpak")) || FParse::Param(FCommandLine::Get(), TEXT("signed"));
@@ -5332,23 +5340,41 @@ FArchive* FPakFile::SetupSignedPakReader(FArchive* ReaderArchive, const TCHAR* F
 #endif
 		if (bShouldCheckSignature)
 		{
-			if (!Decryptor)
-			{
-				Decryptor = MakeUnique<FChunkCacheWorker>(ReaderArchive, Filename);
-			}
+			bNeedsDecryptor = true;
+		}			
+	}
 
-			if (Decryptor->IsValid())
-			{
-				ReaderArchive = new FSignedArchiveReader(ReaderArchive, Decryptor.Get());
-			}
-			else
-			{
-				delete ReaderArchive;
-				return nullptr;
-			}
+	if(bNeedsDecryptor && !Decryptor.IsValid())
+	{
+		TUniquePtr<FArchive> DecryptorReader{ MakeArchive() };
+		if (DecryptorReader.IsValid())
+		{
+			Decryptor = MakeUnique<FChunkCacheWorker>(MoveTemp(DecryptorReader), Filename);
+		}
+
+		if (!Decryptor.IsValid())
+		{
+			return nullptr;
 		}
 	}
-	return ReaderArchive;
+
+	// Now we either have a Decryptor or we don't need it
+	check(!bNeedsDecryptor || Decryptor.IsValid());
+
+	TUniquePtr<FArchive> Archive{ MakeArchive() };
+	if (!Archive.IsValid())
+	{
+		return nullptr;
+	}
+
+	if (Decryptor.IsValid())
+	{
+		return new FSignedArchiveReader(Archive.Release(), Decryptor.Get());
+	}
+	else
+	{
+		return Archive.Release();
+	}
 }
 
 void FPakFile::Initialize(FArchive& Reader, bool bLoadIndex)
@@ -6861,20 +6887,7 @@ bool FPakFile::RecreatePakReaders(IPlatformFile* LowerLevel)
 	// Create a new PakReader *per* instance that was already mapped
 	for (const FArchiveAndLastAccessTime& Reader : Readers)
 	{
-		FArchive* PakReader = nullptr;
-		if (LowerLevel != nullptr)
-		{
-			IFileHandle* PakHandle = LowerLevel->OpenRead(*GetFilename());
-			if (PakHandle)
-			{
-				PakReader = CreatePakReader(*PakHandle, *GetFilename());
-			}
-		}
-		else
-		{
-			PakReader = CreatePakReader(*GetFilename());
-		}
-
+		FArchive* PakReader = CreatePakReader(LowerLevel, *GetFilename());
 		if (!PakReader)
 		{
 			UE_LOG(LogPakFile, Warning, TEXT("Unable to re-create pak \"%s\" handle"), *GetFilename());
@@ -6905,18 +6918,7 @@ FSharedPakReader FPakFile::GetSharedReader(IPlatformFile* LowerLevel)
 		else
 		{
 			// Create a new FArchive reader and pass it to the new handle.
-			if (LowerLevel != nullptr)
-			{
-				IFileHandle* PakHandle = LowerLevel->OpenRead(*GetFilename());
-				if (PakHandle)
-				{
-					PakReader = CreatePakReader(*PakHandle, *GetFilename());
-				}
-			}
-			else
-			{
-				PakReader = CreatePakReader(*GetFilename());
-			}
+			PakReader = CreatePakReader(LowerLevel, *GetFilename());
 
 			if (!PakReader)
 			{
@@ -6948,6 +6950,11 @@ void FPakFile::ReleaseOldReaders(double MaxAgeSeconds)
 			break;
 		}
 	}
+
+	if( Readers.Num() == 0 )
+	{
+		Decryptor.Reset();
+	}
 }
 
 const FPakEntryLocation* FPakFile::FindLocationFromIndex(const FString& FullPath, const FString& MountPoint, const FPathHashIndex& PathHashIndex, uint64 PathHashSeed, int32 PakFileVersion)
diff --git a/Engine/Source/Runtime/PakFile/Private/SignedArchiveReader.cpp b/Engine/Source/Runtime/PakFile/Private/SignedArchiveReader.cpp
--- a/Engine/Source/Runtime/PakFile/Private/SignedArchiveReader.cpp
+++ b/Engine/Source/Runtime/PakFile/Private/SignedArchiveReader.cpp
@@ -30,9 +30,9 @@ DECLARE_DWORD_ACCUMULATOR_STAT(TEXT("FSignedArchiveReader.NumSerializes"), STAT_
 DECLARE_DWORD_ACCUMULATOR_STAT(TEXT("FSignedArchiveReader.NumChunkRequests"), STAT_SignedArchiveReader_NumChunkRequests, STATGROUP_PakFile);
 
 
-FChunkCacheWorker::FChunkCacheWorker(FArchive* InReader, const TCHAR* Filename)
+FChunkCacheWorker::FChunkCacheWorker(TUniquePtr<FArchive> InReader, const TCHAR* Filename)
 	: Thread(nullptr)
-	, Reader(InReader)
+	, Reader(MoveTemp(InReader))
 	, QueuedRequestsEvent(nullptr)
 {
 	Signatures = FPakPlatformFile::GetPakSignatureFile(Filename);
diff --git a/Engine/Source/Runtime/PakFile/Private/SignedArchiveReader.h b/Engine/Source/Runtime/PakFile/Private/SignedArchiveReader.h
--- a/Engine/Source/Runtime/PakFile/Private/SignedArchiveReader.h
+++ b/Engine/Source/Runtime/PakFile/Private/SignedArchiveReader.h
@@ -114,7 +114,7 @@ class FChunkCacheWorker : public FRunnable
 	/** Thread to run the worker FRunnable on */
 	FRunnableThread* Thread;
 	/** Archive reader */
-	FArchive* Reader;
+	TUniquePtr<FArchive> Reader;
 	/** Cached and verified chunks. */
 	FChunkBuffer CachedChunks[MaxCachedChunks];
 	/** Queue of chunks to cache */
@@ -174,7 +174,7 @@ class FChunkCacheWorker : public FRunnable
 
 public:
 
-	FChunkCacheWorker(FArchive* InReader, const TCHAR* Filename);
+	FChunkCacheWorker(TUniquePtr<FArchive> InReader, const TCHAR* Filename);
 	virtual ~FChunkCacheWorker();
 
 	//~ Begin FRunnable Interface.
diff --git a/Engine/Source/Runtime/PakFile/Public/IPlatformFilePak.h b/Engine/Source/Runtime/PakFile/Public/IPlatformFilePak.h
--- a/Engine/Source/Runtime/PakFile/Public/IPlatformFilePak.h
+++ b/Engine/Source/Runtime/PakFile/Public/IPlatformFilePak.h
@@ -808,8 +808,7 @@ private:
 		return 0;
 	}
 
-	FArchive* CreatePakReader(const TCHAR* Filename);
-	FArchive* CreatePakReader(IFileHandle& InHandle, const TCHAR* Filename);
+	FArchive* CreatePakReader(IPlatformFile* LowerLevel, const TCHAR* Filename);
 	FArchive* SetupSignedPakReader(FArchive* Reader, const TCHAR* Filename);
 
 
@@ -1402,7 +1401,7 @@ public:
 		}
 		else
 		{
-			TUniquePtr<FArchive> Reader {CreatePakReader(*GetFilename())};
+			TUniquePtr<FArchive> Reader {CreatePakReader(nullptr, *GetFilename())};
 			Reader->Seek(PakEntry.Offset);
 			FPakEntry SerializedEntry;
 			SerializedEntry.Serialize(*Reader, GetInfo().Version);
