IoDispatcher: Always initialize IoDispatcher

Only return true from FIoDispatcher::IsInitialized() when the script objects chunk exists to keep its behavior consistent. Later, this function will be deprecated in favor of a function on IAsyncPackageLoader.

diff --git a/Engine/Source/Runtime/Core/Private/IO/IoDispatcher.cpp b/Engine/Source/Runtime/Core/Private/IO/IoDispatcher.cpp
--- a/Engine/Source/Runtime/Core/Private/IO/IoDispatcher.cpp
+++ b/Engine/Source/Runtime/Core/Private/IO/IoDispatcher.cpp
@@ -758,18 +758,25 @@ FIoDispatcher::OnSignatureError()
 	return Impl->OnSignatureError();
 }
 
+static bool
+HasScriptObjectsChunk(FIoDispatcher& Dispatcher)
+{
+	static bool bHasScriptObjectsChunk = Dispatcher.DoesChunkExist(CreateIoChunkId(0, 0, EIoChunkType::ScriptObjects));
+	return bHasScriptObjectsChunk;
+}
+
 bool
 FIoDispatcher::IsInitialized()
 {
-	return GIoDispatcher.IsValid();
+	return GIoDispatcher && HasScriptObjectsChunk(*GIoDispatcher);
 }
 
 FIoStatus
 FIoDispatcher::Initialize()
 {
 	LLM_SCOPE(ELLMTag::FileSystem);
+	check(!GIoDispatcher);
 	GIoDispatcher = MakeUnique<FIoDispatcher>();
-
 	return GIoDispatcher->Impl->Initialize();
 }
 
@@ -777,7 +784,6 @@ void
 FIoDispatcher::InitializePostSettings()
 {
 	LLM_SCOPE(ELLMTag::FileSystem);
-
 	check(GIoDispatcher);
 	bool bSuccess = GIoDispatcher->Impl->InitializePostSettings();
 	UE_CLOG(!bSuccess, LogIoDispatcher, Fatal, TEXT("Failed to initialize IoDispatcher"));
diff --git a/Engine/Source/Runtime/Launch/Private/LaunchEngineLoop.cpp b/Engine/Source/Runtime/Launch/Private/LaunchEngineLoop.cpp
--- a/Engine/Source/Runtime/Launch/Private/LaunchEngineLoop.cpp
+++ b/Engine/Source/Runtime/Launch/Private/LaunchEngineLoop.cpp
@@ -69,8 +69,8 @@
 #include "IPlatformFilePak.h"
 #endif
 
-#ifndef USE_IO_DISPATCHER 
-#define USE_IO_DISPATCHER (WITH_IOSTORE_IN_EDITOR || !(IS_PROGRAM || WITH_EDITOR))
+#ifndef USE_IO_DISPATCHER
+#define USE_IO_DISPATCHER (WITH_ENGINE || WITH_IOSTORE_IN_EDITOR || !(IS_PROGRAM || WITH_EDITOR))
 #endif
 #if USE_IO_DISPATCHER
 #include "IO/IoDispatcher.h"
@@ -1807,6 +1807,14 @@ int32 FEngineLoop::PreInitPreStartupScreen(const TCHAR* CmdLine)
 	}
 #endif
 
+#if USE_IO_DISPATCHER
+	if (FIoStatus Status = FIoDispatcher::Initialize(); !Status.IsOk())
+	{
+		UE_LOG(LogInit, Error, TEXT("Failed to initialize I/O dispatcher: '%s'"), *Status.ToString());
+		return 1;
+	}
+#endif
+
 	{
 		SCOPED_BOOT_TIMING("BeginPreInitTextLocalization");
 		BeginPreInitTextLocalization();
@@ -2543,7 +2551,6 @@ PRAGMA_ENABLE_DEPRECATION_WARNINGS
 	}
 
 #if USE_IO_DISPATCHER
-	if (FIoDispatcher::IsInitialized())
 	{
 		SCOPED_BOOT_TIMING("InitIoDispatcher");
 		FIoDispatcher::InitializePostSettings();
diff --git a/Engine/Source/Runtime/PakFile/Private/IPlatformFilePak.cpp b/Engine/Source/Runtime/PakFile/Private/IPlatformFilePak.cpp
--- a/Engine/Source/Runtime/PakFile/Private/IPlatformFilePak.cpp
+++ b/Engine/Source/Runtime/PakFile/Private/IPlatformFilePak.cpp
@@ -266,27 +266,24 @@ void FPakPlatformFile::GetPrunedFilenamesInChunk(const FString& InPakFilename, c
 
 void FPakPlatformFile::GetFilenamesFromIostoreByBlockIndex(const FString& InContainerName, const TArray<int32>& InBlockIndex, TArray<FString>& OutFileList)
 {
-	if (FIoDispatcher::IsInitialized())
+	FPakPlatformFile* PakPlatformFile = static_cast<FPakPlatformFile*>(FPlatformFileManager::Get().FindPlatformFile(FPakPlatformFile::GetTypeName()));
+	if (!PakPlatformFile || !PakPlatformFile->IoDispatcherFileBackend.IsValid())
 	{
-		FPakPlatformFile* PakPlatformFile = static_cast<FPakPlatformFile*>(FPlatformFileManager::Get().FindPlatformFile(FPakPlatformFile::GetTypeName()));
-		if (!PakPlatformFile)
-		{
-			return;
-		}
-		FScopeLock ScopedLock(&PakPlatformFile->PakListCritical);
-		for (const FPakListEntry& PakListEntry : PakPlatformFile->PakFiles)
+		return;
+	}
+	FScopeLock ScopedLock(&PakPlatformFile->PakListCritical);
+	for (const FPakListEntry& PakListEntry : PakPlatformFile->PakFiles)
+	{
+		if (FPaths::GetBaseFilename(PakListEntry.PakFile->PakFilename) == InContainerName)
 		{
-			if (FPaths::GetBaseFilename(PakListEntry.PakFile->PakFilename) == InContainerName)
+			TUniquePtr<FIoStoreReader> IoStoreReader(new FIoStoreReader());
+			FIoStatus Status = IoStoreReader->Initialize(*FPaths::ChangeExtension(PakListEntry.PakFile->PakFilename, TEXT("")), GetRegisteredEncryptionKeys().GetKeys());
+			if (Status.IsOk())
 			{
-				TUniquePtr<FIoStoreReader> IoStoreReader(new FIoStoreReader());
-				FIoStatus Status = IoStoreReader->Initialize(*FPaths::ChangeExtension(PakListEntry.PakFile->PakFilename, TEXT("")), GetRegisteredEncryptionKeys().GetKeys());
-				if (Status.IsOk())
-				{
-					IoStoreReader->GetFilenamesByBlockIndex(InBlockIndex, OutFileList);
-				}
-	
-				break;
+				IoStoreReader->GetFilenamesByBlockIndex(InBlockIndex, OutFileList);
 			}
+	
+			break;
 		}
 	}
 }
@@ -308,26 +305,23 @@ void FPakPlatformFile::GetPrunedFilenamesInPakFile(const FString& InPakFilename,
 
 void FPakPlatformFile::GetFilenamesFromIostoreContainer(const FString& InContainerName, TArray<FString>& OutFileList)
 {
-	if (FIoDispatcher::IsInitialized())
+	FPakPlatformFile* PakPlatformFile = static_cast<FPakPlatformFile*>(FPlatformFileManager::Get().FindPlatformFile(FPakPlatformFile::GetTypeName()));
+	if (!PakPlatformFile || !PakPlatformFile->IoDispatcherFileBackend.IsValid())
 	{
-		FPakPlatformFile* PakPlatformFile = static_cast<FPakPlatformFile*>(FPlatformFileManager::Get().FindPlatformFile(FPakPlatformFile::GetTypeName()));
-		if (!PakPlatformFile)
-		{
-			return;
-		}
-		FScopeLock ScopedLock(&PakPlatformFile->PakListCritical);
-		for (const FPakListEntry& PakListEntry : PakPlatformFile->PakFiles)
+		return;
+	}
+	FScopeLock ScopedLock(&PakPlatformFile->PakListCritical);
+	for (const FPakListEntry& PakListEntry : PakPlatformFile->PakFiles)
+	{
+		if (FPaths::GetBaseFilename(PakListEntry.PakFile->PakFilename) == InContainerName)
 		{
-			if (FPaths::GetBaseFilename(PakListEntry.PakFile->PakFilename) == InContainerName)
+			TUniquePtr<FIoStoreReader> IoStoreReader(new FIoStoreReader());
+			FIoStatus Status = IoStoreReader->Initialize(*FPaths::ChangeExtension(PakListEntry.PakFile->PakFilename, TEXT("")), GetRegisteredEncryptionKeys().GetKeys());
+			if (Status.IsOk())
 			{
-				TUniquePtr<FIoStoreReader> IoStoreReader(new FIoStoreReader());
-				FIoStatus Status = IoStoreReader->Initialize(*FPaths::ChangeExtension(PakListEntry.PakFile->PakFilename, TEXT("")), GetRegisteredEncryptionKeys().GetKeys());
-				if (Status.IsOk())
-				{
-					IoStoreReader->GetFilenames(OutFileList);
-				}
-				break;
+				IoStoreReader->GetFilenames(OutFileList);
 			}
+			break;
 		}
 	}
 }
@@ -4649,7 +4643,7 @@ void FPakPlatformFile::Tick()
 #if CSV_PROFILER
 
 	int64 LocalTotalLoaded = GTotalLoaded;
-	if (FIoDispatcher::IsInitialized())
+	if (IoDispatcherFileBackend.IsValid())
 	{
 		LocalTotalLoaded += FIoDispatcher::Get().GetTotalLoaded();
 	}
@@ -7327,34 +7321,26 @@ bool FPakPlatformFile::Initialize(IPlatformFile* Inner, const TCHAR* CmdLine)
 	FString GlobalUTocPath = FString::Printf(TEXT("%sPaks/global.utoc"), *FPaths::ProjectContentDir());
 	if (!IsRunningCookOnTheFly() && FPlatformFileManager::Get().GetPlatformFile().FileExists(*GlobalUTocPath))
 	{
-		FIoStatus IoDispatcherInitStatus = FIoDispatcher::Initialize();
-		if (IoDispatcherInitStatus.IsOk())
+		FIoDispatcher& IoDispatcher = FIoDispatcher::Get();
+		IoDispatcherFileBackend = CreateIoDispatcherFileBackend();
+		IoDispatcher.Mount(IoDispatcherFileBackend.ToSharedRef());
+		TIoStatusOr<FIoContainerId> IoDispatcherMountStatus = IoDispatcherFileBackend->Mount(*FPaths::ChangeExtension(GlobalUTocPath, TEXT("")), 0, FGuid(), FAES::FAESKey());
+		if (IoDispatcherMountStatus.IsOk())
 		{
-			FIoDispatcher& IoDispatcher = FIoDispatcher::Get();
-			IoDispatcherFileBackend = CreateIoDispatcherFileBackend();
-			IoDispatcher.Mount(IoDispatcherFileBackend.ToSharedRef());
-			TIoStatusOr<FIoContainerId> IoDispatcherMountStatus = IoDispatcherFileBackend->Mount(*FPaths::ChangeExtension(GlobalUTocPath, TEXT("")), 0, FGuid(), FAES::FAESKey());
-			if (IoDispatcherMountStatus.IsOk())
+			UE_LOG(LogPakFile, Display, TEXT("Initialized I/O dispatcher"));
+			IoDispatcher.OnSignatureError().AddLambda([](const FIoSignatureError& Error)
 			{
-				UE_LOG(LogPakFile, Display, TEXT("Initialized I/O dispatcher"));
-				IoDispatcher.OnSignatureError().AddLambda([](const FIoSignatureError& Error)
-				{
-					FPakChunkSignatureCheckFailedData FailedData(Error.ContainerName, TPakChunkHash(), TPakChunkHash(), Error.BlockIndex);
+				FPakChunkSignatureCheckFailedData FailedData(Error.ContainerName, TPakChunkHash(), TPakChunkHash(), Error.BlockIndex);
 #if !PAKHASH_USE_CRC
-					FailedData.ExpectedHash = Error.ExpectedHash;
-					FailedData.ReceivedHash = Error.ActualHash;
+				FailedData.ExpectedHash = Error.ExpectedHash;
+				FailedData.ReceivedHash = Error.ActualHash;
 #endif
-					FPakPlatformFile::BroadcastPakChunkSignatureCheckFailure(FailedData);
-				});
-			}
-			else
-			{
-				UE_LOG(LogPakFile, Error, TEXT("Failed to mount I/O dispatcher container: '%s'"), *IoDispatcherMountStatus.Status().ToString());
-			}
+				FPakPlatformFile::BroadcastPakChunkSignatureCheckFailure(FailedData);
+			});
 		}
 		else
 		{
-			UE_LOG(LogPakFile, Error, TEXT("Failed to initialize I/O dispatcher: '%s'"), *IoDispatcherInitStatus.ToString());
+			UE_LOG(LogPakFile, Error, TEXT("Failed to mount I/O dispatcher container: '%s'"), *IoDispatcherMountStatus.Status().ToString());
 		}
 	}
 
@@ -7678,7 +7664,7 @@ bool FPakPlatformFile::Unmount(const TCHAR* InPakFilename)
 			}
 		}
 
-		if (FIoDispatcher::IsInitialized())
+		if (IoDispatcherFileBackend.IsValid())
 		{
 			FString ContainerPath = FPaths::ChangeExtension(InPakFilename, FString());
 			FIoStatus Status = FIoDispatcher::Get().Unmount(*ContainerPath);
diff --git a/Engine/Source/Runtime/StorageServerClient/Private/StorageServerPlatformFile.cpp b/Engine/Source/Runtime/StorageServerClient/Private/StorageServerPlatformFile.cpp
--- a/Engine/Source/Runtime/StorageServerClient/Private/StorageServerPlatformFile.cpp
+++ b/Engine/Source/Runtime/StorageServerClient/Private/StorageServerPlatformFile.cpp
@@ -278,25 +278,17 @@ bool FStorageServerPlatformFile::Initialize(IPlatformFile* Inner, const TCHAR* C
 		{
 			if (SendGetFileListMessage())
 			{
-				FIoStatus IoDispatcherInitStatus = FIoDispatcher::Initialize();
-				if (IoDispatcherInitStatus.IsOk())
-				{
-					FIoDispatcher& IoDispatcher = FIoDispatcher::Get();
-					TSharedRef<FStorageServerIoDispatcherBackend> IoDispatcherBackend = MakeShared<FStorageServerIoDispatcherBackend>(*Connection.Get());
-					IoDispatcher.Mount(IoDispatcherBackend);
+				FIoDispatcher& IoDispatcher = FIoDispatcher::Get();
+				TSharedRef<FStorageServerIoDispatcherBackend> IoDispatcherBackend = MakeShared<FStorageServerIoDispatcherBackend>(*Connection.Get());
+				IoDispatcher.Mount(IoDispatcherBackend);
 #if WITH_COTF
-					if (IsRunningCookOnTheFly())
-					{
-						UE::Cook::ICookOnTheFlyModule& CookOnTheFlyModule = FModuleManager::LoadModuleChecked<UE::Cook::ICookOnTheFlyModule>(TEXT("CookOnTheFly"));
-						CookOnTheFlyModule.GetServerConnection().OnMessage().AddRaw(this, &FStorageServerPlatformFile::OnCookOnTheFlyMessage);
-					}
-#endif					
-					bResult = true;
-				}
-				else
+				if (IsRunningCookOnTheFly())
 				{
-					UE_LOG(LogStorageServerPlatformFile, Fatal, TEXT("Failed to initialize IoDispatcher with Zen host '%s'"), *Connection->GetHostAddr());
+					UE::Cook::ICookOnTheFlyModule& CookOnTheFlyModule = FModuleManager::LoadModuleChecked<UE::Cook::ICookOnTheFlyModule>(TEXT("CookOnTheFly"));
+					CookOnTheFlyModule.GetServerConnection().OnMessage().AddRaw(this, &FStorageServerPlatformFile::OnCookOnTheFlyMessage);
 				}
+#endif
+				bResult = true;
 			}
 			else
 			{
