Added missing include
#rnx

diff --git a/Engine/Source/Runtime/PakFile/Private/FilePackageStore.cpp b/Engine/Source/Runtime/PakFile/Private/FilePackageStore.cpp
--- a/Engine/Source/Runtime/PakFile/Private/FilePackageStore.cpp
+++ b/Engine/Source/Runtime/PakFile/Private/FilePackageStore.cpp
@@ -8,6 +8,7 @@
 #include "Misc/ScopeRWLock.h"
 #include "ProfilingDebugging/CountersTrace.h"
 #include "Internationalization/PackageLocalizationManager.h"
+#include "Misc/ScopeLock.h"
 
 //PRAGMA_DISABLE_OPTIMIZATION
 
