Don't discard decryptor if there are existing loaned-out pak readers.

diff --git a/Engine/Source/Runtime/PakFile/Private/IPlatformFilePak.cpp b/Engine/Source/Runtime/PakFile/Private/IPlatformFilePak.cpp
--- a/Engine/Source/Runtime/PakFile/Private/IPlatformFilePak.cpp
+++ b/Engine/Source/Runtime/PakFile/Private/IPlatformFilePak.cpp
@@ -6875,6 +6875,11 @@ bool FPakFile::RecreatePakReaders(IPlatformFile* LowerLevel)
 {
 	FScopeLock ScopedLock(&CriticalSection);
 
+	if (CurrentlyUsedReaders > 0)
+	{
+		UE_LOG(LogPakFile, Error, TEXT("Recreating pak readers while we have readers loaned out, this may be lead to crashes or decryption problems."));
+	}
+
 	// need to reset the decryptor as it will hold a pointer to the first created pak reader
 	Decryptor.Reset();
 
@@ -6921,6 +6926,7 @@ FSharedPakReader FPakFile::GetSharedReader(IPlatformFile* LowerLevel)
 				UE_LOG(LogPakFile, Warning, TEXT("Unable to create pak \"%s\" handle"), *GetFilename());
 			}
 		}
+		++CurrentlyUsedReaders;
 	}
 
 	return FSharedPakReader(PakReader, this);
@@ -6929,6 +6935,7 @@ FSharedPakReader FPakFile::GetSharedReader(IPlatformFile* LowerLevel)
 void FPakFile::ReturnSharedReader(FArchive* Archive)
 {
 	FScopeLock ScopedLock(&CriticalSection);
+	--CurrentlyUsedReaders;
 	Readers.Push(FArchiveAndLastAccessTime{ TUniquePtr<FArchive>{Archive }, FPlatformTime::Seconds()});
 }
 
@@ -6947,7 +6954,7 @@ void FPakFile::ReleaseOldReaders(double MaxAgeSeconds)
 		}
 	}
 
-	if( Readers.Num() == 0 )
+	if (Readers.Num() == 0 && CurrentlyUsedReaders == 0)
 	{
 		Decryptor.Reset();
 	}
diff --git a/Engine/Source/Runtime/PakFile/Public/IPlatformFilePak.h b/Engine/Source/Runtime/PakFile/Public/IPlatformFilePak.h
--- a/Engine/Source/Runtime/PakFile/Public/IPlatformFilePak.h
+++ b/Engine/Source/Runtime/PakFile/Public/IPlatformFilePak.h
@@ -737,6 +737,8 @@ private:
 	TUniquePtr<class FChunkCacheWorker> Decryptor;
 	/** List of readers and when they were last used. */
 	TArray<FArchiveAndLastAccessTime> Readers;
+	/** How many readers have been loaned out and not yet returned. If this is >0 we should not destroy the decryptor. */
+	int32 CurrentlyUsedReaders = 0;
 	/** Critical section for accessing Readers. */
 	FCriticalSection CriticalSection;
 	/** Pak file info (trailer). */
