Implemented -checkpak support for iostore containers
#rnx

diff --git a/Engine/Source/Runtime/Core/Internal/IO/IoDispatcherFileBackend.h b/Engine/Source/Runtime/Core/Internal/IO/IoDispatcherFileBackend.h
--- a/Engine/Source/Runtime/Core/Internal/IO/IoDispatcherFileBackend.h
+++ b/Engine/Source/Runtime/Core/Internal/IO/IoDispatcherFileBackend.h
@@ -28,7 +28,7 @@ class FFileIoStoreReader
 public:
 	FFileIoStoreReader(IPlatformFileIoStore& InPlatformImpl, FFileIoStoreStats& InStats);
 	~FFileIoStoreReader();
-	FIoStatus Initialize(const TCHAR* ContainerPath, int32 Order);
+	FIoStatus Initialize(const TCHAR* InTocFilePath, int32 Order);
 	uint32 GetContainerInstanceId() const
 	{
 		return ContainerFile.ContainerInstanceId;
@@ -106,8 +106,8 @@ public:
 	FFileIoStore(TUniquePtr<IPlatformFileIoStore>&& PlatformImpl);
 	~FFileIoStore();
 	void Initialize(TSharedRef<const FIoDispatcherBackendContext> Context) override;
-	TIoStatusOr<FIoContainerHeader> Mount(const TCHAR* ContainerPath, int32 Order, const FGuid& EncryptionKeyGuid, const FAES::FAESKey& EncryptionKey) override;
-	bool Unmount(const TCHAR* ContainerPath) override;
+	TIoStatusOr<FIoContainerHeader> Mount(const TCHAR* InTocPath, int32 Order, const FGuid& EncryptionKeyGuid, const FAES::FAESKey& EncryptionKey) override;
+	bool Unmount(const TCHAR* InTocPath) override;
 	bool Resolve(FIoRequestImpl* Request) override;
 	void CancelIoRequest(FIoRequestImpl* Request) override;
 	void UpdatePriorityForIoRequest(FIoRequestImpl* Request) override;
diff --git a/Engine/Source/Runtime/Core/Internal/IO/IoStore.h b/Engine/Source/Runtime/Core/Internal/IO/IoStore.h
--- a/Engine/Source/Runtime/Core/Internal/IO/IoStore.h
+++ b/Engine/Source/Runtime/Core/Internal/IO/IoStore.h
@@ -235,7 +235,7 @@ struct FIoStoreTocResource
 
 	TArray<uint8> DirectoryIndexBuffer;
 
-	UE_NODISCARD static FIoStatus Read(const TCHAR* TocFilePath, EIoStoreTocReadOptions ReadOptions, FIoStoreTocResource& OutTocResource);
+	UE_NODISCARD CORE_API static FIoStatus Read(const TCHAR* TocFilePath, EIoStoreTocReadOptions ReadOptions, FIoStoreTocResource& OutTocResource);
 
 	UE_NODISCARD static TIoStatusOr<uint64> Write(const TCHAR* TocFilePath, FIoStoreTocResource& TocResource, const FIoContainerSettings& ContainerSettings, const FIoStoreWriterSettings& WriterSettings);
 
diff --git a/Engine/Source/Runtime/Core/Private/IO/IoDispatcherFileBackend.cpp b/Engine/Source/Runtime/Core/Private/IO/IoDispatcherFileBackend.cpp
--- a/Engine/Source/Runtime/Core/Private/IO/IoDispatcherFileBackend.cpp
+++ b/Engine/Source/Runtime/Core/Private/IO/IoDispatcherFileBackend.cpp
@@ -678,20 +678,22 @@ uint64 FFileIoStoreReader::GetTocAllocatedSize() const
 		ContainerFile.BlockSignatureHashes.GetAllocatedSize();
 }
 
-FIoStatus FFileIoStoreReader::Initialize(const TCHAR* InContainerPath, int32 InOrder)
+FIoStatus FFileIoStoreReader::Initialize(const TCHAR* InTocFilePath, int32 InOrder)
 {
-	IPlatformFile& Ipf = FPlatformFileManager::Get().GetPlatformFile();
+	FStringView ContainerPathView(InTocFilePath);
+	if (!ContainerPathView.EndsWith(TEXT(".utoc")))
+	{
+		return FIoStatusBuilder(EIoErrorCode::FileOpenFailed) << TEXT("Expected .utoc extension on container path '") << InTocFilePath << TEXT("'");
+	}
+	FStringView BasePathView = ContainerPathView.LeftChop(5);
 
-	TStringBuilder<256> TocFilePath;
-	TocFilePath.Append(InContainerPath);
-	TocFilePath.Append(TEXT(".utoc"));
-	ContainerFile.FilePath = TocFilePath;
+	IPlatformFile& Ipf = FPlatformFileManager::Get().GetPlatformFile();
 
-	UE_LOG(LogIoDispatcher, Display, TEXT("Reading toc: %s"), *TocFilePath);
+	UE_LOG(LogIoDispatcher, Display, TEXT("Reading toc: %s"), InTocFilePath);
 
 	TUniquePtr<FIoStoreTocResource> TocResourcePtr = MakeUnique<FIoStoreTocResource>();
 	FIoStoreTocResource& TocResource = *TocResourcePtr;
-	FIoStatus Status = FIoStoreTocResource::Read(*TocFilePath, EIoStoreTocReadOptions::Default, TocResource);
+	FIoStatus Status = FIoStoreTocResource::Read(InTocFilePath, EIoStoreTocReadOptions::Default, TocResource);
 	if (!Status.IsOk())
 	{
 		return Status;
@@ -703,7 +705,7 @@ FIoStatus FFileIoStoreReader::Initialize(const TCHAR* InContainerPath, int32 InO
 	{
 		FFileIoStoreContainerFilePartition& Partition = ContainerFile.Partitions[PartitionIndex];
 		TStringBuilder<256> ContainerFilePath;
-		ContainerFilePath.Append(InContainerPath);
+		ContainerFilePath.Append(BasePathView);
 		if (PartitionIndex > 0)
 		{
 			ContainerFilePath.Appendf(TEXT("_s%d"), PartitionIndex);
@@ -731,7 +733,7 @@ FIoStatus FFileIoStoreReader::Initialize(const TCHAR* InContainerPath, int32 InO
 	}
 	else
 	{
-		UE_LOG(LogIoDispatcher, Warning, TEXT("Falling back to imperfect hashmap for container '%s'"), *TocFilePath);
+		UE_LOG(LogIoDispatcher, Warning, TEXT("Falling back to imperfect hashmap for container '%s'"), InTocFilePath);
 		for (uint32 ChunkIndex = 0; ChunkIndex < TocResource.Header.TocEntryCount; ++ChunkIndex)
 		{
 			TocImperfectHashMapFallback.Add(TocResource.ChunkIds[ChunkIndex], TocResource.ChunkOffsetLengths[ChunkIndex]);
@@ -1250,10 +1252,10 @@ void FFileIoStore::Initialize(TSharedRef<const FIoDispatcherBackendContext> InCo
 	Thread = FRunnableThread::Create(this, TEXT("IoService"), 0, TPri_AboveNormal);
 }
 
-TIoStatusOr<FIoContainerHeader> FFileIoStore::Mount(const TCHAR* ContainerPath, int32 Order, const FGuid& EncryptionKeyGuid, const FAES::FAESKey& EncryptionKey)
+TIoStatusOr<FIoContainerHeader> FFileIoStore::Mount(const TCHAR* InTocPath, int32 Order, const FGuid& EncryptionKeyGuid, const FAES::FAESKey& EncryptionKey)
 {
 	TUniquePtr<FFileIoStoreReader> Reader(new FFileIoStoreReader(*PlatformImpl, Stats));
-	FIoStatus IoStatus = Reader->Initialize(ContainerPath, Order);
+	FIoStatus IoStatus = Reader->Initialize(InTocPath, Order);
 	if (!IoStatus.IsOk())
 	{
 		return IoStatus;
@@ -1268,7 +1270,7 @@ TIoStatusOr<FIoContainerHeader> FFileIoStore::Mount(const TCHAR* ContainerPath,
 		else
 		{
 			return FIoStatus(EIoErrorCode::InvalidEncryptionKey, *FString::Printf(TEXT("Invalid encryption key '%s' (container '%s', encryption key '%s')"),
-				*EncryptionKeyGuid.ToString(), *FPaths::GetBaseFilename(ContainerPath), *Reader->GetEncryptionKeyGuid().ToString()));
+				*EncryptionKeyGuid.ToString(), InTocPath, *Reader->GetEncryptionKeyGuid().ToString()));
 		}
 	}
 
@@ -1295,23 +1297,21 @@ TIoStatusOr<FIoContainerHeader> FFileIoStore::Mount(const TCHAR* ContainerPath,
 			return A->GetContainerInstanceId() > B->GetContainerInstanceId();
 		});
 		IoStoreReaders.Insert(MoveTemp(Reader), InsertionIndex);
-		UE_LOG(LogIoDispatcher, Display, TEXT("Mounting container '%s' in location slot %d"), ContainerPath, InsertionIndex);
+		UE_LOG(LogIoDispatcher, Display, TEXT("Mounting container '%s' in location slot %d"), InTocPath, InsertionIndex);
 	}
 
 	return ContainerHeader;
 }
 
-bool FFileIoStore::Unmount(const TCHAR* ContainerPath)
+bool FFileIoStore::Unmount(const TCHAR* InTocPath)
 {
 	FWriteScopeLock _(IoStoreReadersLock);
 	
-	FString FilePathToUnmount = FPaths::SetExtension(ContainerPath, FString(".utoc"));
-
 	for (int32 Idx = 0; Idx < IoStoreReaders.Num(); ++Idx)
 	{
-		if (IoStoreReaders[Idx]->GetContainerFile().FilePath == FilePathToUnmount)
+		if (IoStoreReaders[Idx]->GetContainerFile().FilePath == InTocPath)
 		{
-			UE_LOG(LogIoDispatcher, Display, TEXT("Unmounting container '%s'"), *FPaths::GetBaseFilename(ContainerPath));
+			UE_LOG(LogIoDispatcher, Display, TEXT("Unmounting container '%s'"), InTocPath);
 
 			// Cancel pending I/O requests trying to read from the container
 			for (const FFileIoStoreContainerFilePartition& Partition : IoStoreReaders[Idx]->GetContainerFile().Partitions)
@@ -1326,7 +1326,7 @@ bool FFileIoStore::Unmount(const TCHAR* ContainerPath)
 		}
 	}
 
-	UE_LOG(LogIoDispatcher, Display, TEXT("Failed to unmount container '%s'"), *FPaths::GetBaseFilename(ContainerPath));
+	UE_LOG(LogIoDispatcher, Display, TEXT("Failed to unmount container '%s'"), InTocPath);
 	
 	return false;
 }
diff --git a/Engine/Source/Runtime/PakFile/Private/IPlatformFilePak.cpp b/Engine/Source/Runtime/PakFile/Private/IPlatformFilePak.cpp
--- a/Engine/Source/Runtime/PakFile/Private/IPlatformFilePak.cpp
+++ b/Engine/Source/Runtime/PakFile/Private/IPlatformFilePak.cpp
@@ -40,6 +40,7 @@
 #include "IO/IoContainerHeader.h"
 #include "FilePackageStore.h"
 #include "Compression/OodleDataCompression.h"
+#include "IO/IoStore.h"
 
 DEFINE_LOG_CATEGORY(LogPakFile);
 
@@ -5373,6 +5374,12 @@ FArchive* FPakFile::CreatePakReader(IPlatformFile* LowerLevel, const TCHAR* File
 	}
 }
 
+bool ShouldCheckPak()
+{
+	static bool bShouldCheckPak = FParse::Param(FCommandLine::Get(), TEXT("checkpak"));
+	return bShouldCheckPak;
+}
+
 void FPakFile::Initialize(FArchive& Reader, bool bLoadIndex)
 {
 	CachedTotalSize = Reader.TotalSize();
@@ -5422,7 +5429,7 @@ void FPakFile::Initialize(FArchive& Reader, bool bLoadIndex)
 				LoadIndex(Reader);
 			}
 
-			if (FParse::Param(FCommandLine::Get(), TEXT("checkpak")))
+			if (ShouldCheckPak())
 			{
 				ensure(Check());
 			}
@@ -6621,8 +6628,6 @@ bool FPakFile::Check()
 	int32 ErrorCount = 0;
 	int32 FileCount = 0;
 
-	bool bSuccess = true;
-
 	// If the pak file is signed, we can do a fast check by just reading a single byte from the start of
 	// each signing block. The signed archive reader will bring in that whole signing block and compare
 	// against the signature table and fire the handler
@@ -6633,9 +6638,9 @@ bool FPakFile::Check()
 
 		{
 			FScopeLock Lock(&HandlerData.Lock);
-			DelegateHandle = HandlerData.ChunkSignatureCheckFailedDelegate.AddLambda([&bSuccess](const FPakChunkSignatureCheckFailedData&)
+			DelegateHandle = HandlerData.ChunkSignatureCheckFailedDelegate.AddLambda([&ErrorCount](const FPakChunkSignatureCheckFailedData&)
 			{
-				bSuccess = false;
+				++ErrorCount;
 			});
 		}
 
@@ -6723,6 +6728,89 @@ bool FPakFile::Check()
 	return ErrorCount == 0;
 }
 
+bool CheckIoStoreContainerBlockSignatures(const TCHAR* InContainerPath)
+{
+	TRACE_CPUPROFILER_EVENT_SCOPE(CheckIoStoreContainerBlockSignatures);
+	UE_LOG(LogPakFile, Display, TEXT("Checking container file \"%s\"..."), InContainerPath);
+	double StartTime = FPlatformTime::Seconds();
+
+	FIoStoreTocResource TocResource;
+	FIoStatus Status = FIoStoreTocResource::Read(InContainerPath, EIoStoreTocReadOptions::Default, TocResource);
+	if (!Status.IsOk())
+	{
+		UE_LOG(LogPakFile, Error, TEXT("Failed reading toc file \"%s\"."), InContainerPath);
+		return false;
+	}
+
+	if (TocResource.ChunkBlockSignatures.Num() != TocResource.CompressionBlocks.Num())
+	{
+		UE_LOG(LogPakFile, Error, TEXT("Toc file \"%s\" doesn't contain any chunk block signatures."), InContainerPath);
+		return false;
+	}
+
+	TUniquePtr<FArchive> ContainerFileReader;
+	int32 LastPartitionIndex = -1;
+	IPlatformFile& Ipf = FPlatformFileManager::Get().GetPlatformFile();
+	TArray<uint8> BlockBuffer;
+	BlockBuffer.SetNum(static_cast<int32>(TocResource.Header.CompressionBlockSize));
+	const int32 BlockCount = TocResource.CompressionBlocks.Num();
+	int32 ErrorCount = 0;
+	FString ContainerBasePath = FPaths::ChangeExtension(InContainerPath, TEXT(""));
+	TStringBuilder<256> UcasFilePath;
+	for (int32 BlockIndex = 0; BlockIndex < BlockCount; ++BlockIndex)
+	{
+		const FIoStoreTocCompressedBlockEntry& CompressionBlockEntry = TocResource.CompressionBlocks[BlockIndex];
+		uint64 BlockRawSize = Align(CompressionBlockEntry.GetCompressedSize(), FAES::AESBlockSize);
+		check(BlockRawSize <= TocResource.Header.CompressionBlockSize);
+		const int32 PartitionIndex = int32(CompressionBlockEntry.GetOffset() / TocResource.Header.PartitionSize);
+		const uint64 PartitionRawOffset = CompressionBlockEntry.GetOffset() % TocResource.Header.PartitionSize;
+		if (PartitionIndex != LastPartitionIndex)
+		{
+			UcasFilePath.Reset();
+			UcasFilePath.Append(ContainerBasePath);
+			if (PartitionIndex > 0)
+			{
+				UcasFilePath.Append(FString::Printf(TEXT("_s%d"), PartitionIndex));
+			}
+			UcasFilePath.Append(TEXT(".ucas"));
+			IFileHandle* ContainerFileHandle = Ipf.OpenRead(*UcasFilePath, /* allowwrite */ false);
+			if (!ContainerFileHandle)
+			{
+				UE_LOG(LogPakFile, Error, TEXT("Failed opening container file \"%s\"."), *UcasFilePath);
+				return false;
+			}
+			ContainerFileReader.Reset(new FArchiveFileReaderGeneric(ContainerFileHandle, *UcasFilePath, ContainerFileHandle->Size(), 256 << 10));
+			LastPartitionIndex = PartitionIndex;
+		}
+		ContainerFileReader->Seek(PartitionRawOffset);
+		ContainerFileReader->Precache(PartitionRawOffset, 0); // Without this buffering won't work due to the first read after a seek always being uncached
+		ContainerFileReader->Serialize(BlockBuffer.GetData(), BlockRawSize);
+		FSHAHash BlockHash;
+		FSHA1::HashBuffer(BlockBuffer.GetData(), BlockRawSize, BlockHash.Hash);
+		if (TocResource.ChunkBlockSignatures[BlockIndex] != BlockHash)
+		{
+			UE_LOG(LogPakFile, Warning, TEXT("Hash mismatch for block [%i/%i]! Expected %s, Received %s"), BlockIndex, BlockCount, *TocResource.ChunkBlockSignatures[BlockIndex].ToString(), *BlockHash.ToString());
+
+			FPakChunkSignatureCheckFailedData Data(*UcasFilePath, TPakChunkHash(), TPakChunkHash(), BlockIndex);
+#if PAKHASH_USE_CRC
+			Data.ExpectedHash = GetTypeHash(TocResource.ChunkBlockSignatures[BlockIndex]);
+			Data.ReceivedHash = GetTypeHash(BlockHash);
+#else
+			Data.ExpectedHash = TocResource.ChunkBlockSignatures[BlockIndex];
+			Data.ReceivedHash = BlockHash;
+#endif
+			FPakPlatformFile::BroadcastPakChunkSignatureCheckFailure(Data);
+			++ErrorCount;
+		}
+	}
+
+	double EndTime = FPlatformTime::Seconds();
+	double ElapsedTime = EndTime - StartTime;
+	UE_LOG(LogPakFile, Display, TEXT("Container file \"%s\" checked in %.2fs"), InContainerPath, ElapsedTime);
+
+	return ErrorCount == 0;
+}
+
 #if DO_CHECK
 /**
 * FThreadCheckingArchiveProxy - checks that inner archive is only used from the specified thread ID
@@ -7332,6 +7420,11 @@ bool FPakPlatformFile::Initialize(IPlatformFile* Inner, const TCHAR* CmdLine)
 	const bool bForceIoStore = WITH_IOSTORE_IN_EDITOR && FParse::Param(CmdLine, TEXT("UseIoStore"));
 	if (bShouldMountGlobal || bForceIoStore)
 	{
+		if (ShouldCheckPak())
+		{
+			ensure(CheckIoStoreContainerBlockSignatures(*GlobalUTocPath));
+		}
+
 		FIoDispatcher& IoDispatcher = FIoDispatcher::Get();
 		IoDispatcherFileBackend = CreateIoDispatcherFileBackend();
 		IoDispatcher.Mount(IoDispatcherFileBackend.ToSharedRef());
@@ -7343,14 +7436,17 @@ bool FPakPlatformFile::Initialize(IPlatformFile* Inner, const TCHAR* CmdLine)
 
 		if (bShouldMountGlobal)
 		{
-			TIoStatusOr<FIoContainerHeader> IoDispatcherMountStatus = IoDispatcherFileBackend->Mount(*FPaths::ChangeExtension(GlobalUTocPath, TEXT("")), 0, FGuid(), FAES::FAESKey());
+			TIoStatusOr<FIoContainerHeader> IoDispatcherMountStatus = IoDispatcherFileBackend->Mount(*GlobalUTocPath, 0, FGuid(), FAES::FAESKey());
 			if (IoDispatcherMountStatus.IsOk())
 			{
 				UE_LOG(LogPakFile, Display, TEXT("Initialized I/O dispatcher file backend. Mounted the global container: %s"), *GlobalUTocPath);
 				IoDispatcher.OnSignatureError().AddLambda([](const FIoSignatureError& Error)
 				{
 					FPakChunkSignatureCheckFailedData FailedData(Error.ContainerName, TPakChunkHash(), TPakChunkHash(), Error.BlockIndex);
-#if !PAKHASH_USE_CRC
+#if PAKHASH_USE_CRC
+					FailedData.ExpectedHash = GetTypeHash(Error.ExpectedHash);
+					FailedData.ReceivedHash = GetTypeHash(Error.ActualHash);
+#else
 					FailedData.ExpectedHash = Error.ExpectedHash;
 					FailedData.ReceivedHash = Error.ActualHash;
 #endif
@@ -7629,18 +7725,30 @@ bool FPakPlatformFile::Mount(const TCHAR* InPakFilename, uint32 PakOrder, const
 				}
 			}
 
-			FString ContainerPath = FPaths::ChangeExtension(InPakFilename, FString());
-			TIoStatusOr<FIoContainerHeader> MountResult = IoDispatcherFileBackend->Mount(*ContainerPath, PakOrder, EncryptionKeyGuid, EncryptionKey);
-			if (MountResult.IsOk())
+			FString UtocPath = FPaths::ChangeExtension(InPakFilename, TEXT(".utoc"));
+			if (FPlatformFileManager::Get().GetPlatformFile().FileExists(*UtocPath))
 			{
-				UE_LOG(LogPakFile, Display, TEXT("Mounted IoStore container \"%s\""), *ContainerPath);
-				Pak->IoContainerHeader = MakeUnique<FIoContainerHeader>(MountResult.ConsumeValueOrDie());
-				FilePackageStore->Mount(Pak->IoContainerHeader.Get(), PakOrder);
+				if (ShouldCheckPak())
+				{
+					ensure(CheckIoStoreContainerBlockSignatures(*UtocPath));
+				}
+				TIoStatusOr<FIoContainerHeader> MountResult = IoDispatcherFileBackend->Mount(*UtocPath, PakOrder, EncryptionKeyGuid, EncryptionKey);
+				if (MountResult.IsOk())
+				{
+					UE_LOG(LogPakFile, Display, TEXT("Mounted IoStore container \"%s\""), *UtocPath);
+					Pak->IoContainerHeader = MakeUnique<FIoContainerHeader>(MountResult.ConsumeValueOrDie());
+					FilePackageStore->Mount(Pak->IoContainerHeader.Get(), PakOrder);
+				}
+				else
+				{
+					bIoStoreSuccess = false;
+					UE_LOG(LogPakFile, Warning, TEXT("Failed to mount IoStore container \"%s\" [%s]"), *UtocPath);
+				}
 			}
 			else
 			{
 				bIoStoreSuccess = false;
-				UE_LOG(LogPakFile, Warning, TEXT("Failed to mount IoStore container \"%s\" [%s]"), *ContainerPath, *MountResult.Status().ToString());
+				UE_LOG(LogPakFile, Warning, TEXT("IoStore container \"%s\" not found"), *UtocPath);
 			}
 		}
 
