SignedArchiveReader: Make sure to write some data to the output buffer even if the signatures are missing. Otherwise we could end up with a valid pak file magic number from a previous read and incorrectly consider the pak file valid.
#rnx

diff --git a/Engine/Source/Runtime/PakFile/Private/SignedArchiveReader.cpp b/Engine/Source/Runtime/PakFile/Private/SignedArchiveReader.cpp
--- a/Engine/Source/Runtime/PakFile/Private/SignedArchiveReader.cpp
+++ b/Engine/Source/Runtime/PakFile/Private/SignedArchiveReader.cpp
@@ -276,6 +276,10 @@ bool FChunkCacheWorker::CheckSignature(const FChunkRequest& ChunkInfo)
 			FPakPlatformFile::BroadcastPakChunkSignatureCheckFailure(Data);
 		}
 	}
+	else
+	{
+		FMemory::Memset(ChunkInfo.Buffer->Data, 0xcd, ChunkInfo.Size);
+	}
 	
 	return bChunkHashesMatch;
 }
