Move IoDispatcherFileBackend from Core to PakFile
#rnx

diff --git a/Engine/Source/Runtime/Core/Internal/IO/IoDispatcherFileBackend.h b/Engine/Source/Runtime/Core/Internal/IO/IoDispatcherFileBackend.h
deleted file mode 100644
diff --git a/Engine/Source/Runtime/Core/Internal/IO/IoDispatcherFileBackendTypes.h b/Engine/Source/Runtime/Core/Internal/IO/IoDispatcherFileBackendTypes.h
deleted file mode 100644
diff --git a/Engine/Source/Runtime/Core/Internal/IO/IoStore.h b/Engine/Source/Runtime/Core/Internal/IO/IoStore.h
--- a/Engine/Source/Runtime/Core/Internal/IO/IoStore.h
+++ b/Engine/Source/Runtime/Core/Internal/IO/IoStore.h
@@ -241,5 +241,5 @@ struct FIoStoreTocResource
 
 	UE_NODISCARD static TIoStatusOr<uint64> Write(const TCHAR* TocFilePath, FIoStoreTocResource& TocResource, const FIoContainerSettings& ContainerSettings, const FIoStoreWriterSettings& WriterSettings);
 
-	static uint64 HashChunkIdWithSeed(int32 Seed, const FIoChunkId& ChunkId);
+	CORE_API static uint64 HashChunkIdWithSeed(int32 Seed, const FIoChunkId& ChunkId);
 };
diff --git a/Engine/Source/Runtime/Core/Private/GenericPlatform/GenericPlatformIoDispatcher.cpp b/Engine/Source/Runtime/Core/Private/GenericPlatform/GenericPlatformIoDispatcher.cpp
deleted file mode 100644
diff --git a/Engine/Source/Runtime/Core/Private/GenericPlatform/GenericPlatformIoDispatcher.h b/Engine/Source/Runtime/Core/Private/GenericPlatform/GenericPlatformIoDispatcher.h
deleted file mode 100644
diff --git a/Engine/Source/Runtime/Core/Private/IO/IoDispatcher.cpp b/Engine/Source/Runtime/Core/Private/IO/IoDispatcher.cpp
--- a/Engine/Source/Runtime/Core/Private/IO/IoDispatcher.cpp
+++ b/Engine/Source/Runtime/Core/Private/IO/IoDispatcher.cpp
@@ -19,7 +19,6 @@
 #include "ProfilingDebugging/CsvProfiler.h"
 #include "Containers/Ticker.h"
 #include "IO/IoDispatcherBackend.h"
-#include "IO/IoDispatcherFileBackendTypes.h"
 #include "Hash/Blake3.h"
 
 DEFINE_LOG_CATEGORY(LogIoDispatcher);
diff --git a/Engine/Source/Runtime/Core/Private/IO/IoDispatcherFileBackend.cpp b/Engine/Source/Runtime/Core/Private/IO/IoDispatcherFileBackend.cpp
deleted file mode 100644
diff --git a/Engine/Source/Runtime/Core/Public/IO/IoDispatcherBackend.h b/Engine/Source/Runtime/Core/Public/IO/IoDispatcherBackend.h
--- a/Engine/Source/Runtime/Core/Public/IO/IoDispatcherBackend.h
+++ b/Engine/Source/Runtime/Core/Public/IO/IoDispatcherBackend.h
@@ -117,14 +117,3 @@ struct IIoDispatcherBackend
 	virtual FIoRequestImpl* GetCompletedRequests() = 0;
 	virtual TIoStatusOr<FIoMappedRegion> OpenMapped(const FIoChunkId& ChunkId, const FIoReadOptions& Options) = 0;
 };
-
-struct IIoDispatcherFileBackend
-	: public IIoDispatcherBackend
-{
-	virtual TIoStatusOr<FIoContainerHeader> Mount(const TCHAR* ContainerPath, int32 Order, const FGuid& EncryptionKeyGuid, const FAES::FAESKey& EncryptionKey) = 0;
-	virtual bool Unmount(const TCHAR* ContainerPath) = 0;
-	virtual void ReopenAllFileHandles() = 0;
-};
-
-CORE_API TSharedRef<IIoDispatcherFileBackend> CreateIoDispatcherFileBackend();
-
diff --git a/Engine/Source/Runtime/PakFile/Internal/IoDispatcherFileBackendTypes.h b/Engine/Source/Runtime/PakFile/Internal/IoDispatcherFileBackendTypes.h
new file mode 100644
--- /dev/null
+++ b/Engine/Source/Runtime/PakFile/Internal/IoDispatcherFileBackendTypes.h
@@ -0,0 +1,1046 @@
+// Copyright Epic Games, Inc. All Rights Reserved.
+
+#pragma once
+
+#include "IO/IoStore.h"
+#include "IO/IoDispatcherBackend.h"
+#include "ProfilingDebugging/CsvProfiler.h"
+#include "ProfilingDebugging/CountersTrace.h"
+#include "Containers/Ticker.h"
+#include "Modules/ModuleInterface.h"
+
+#ifndef PLATFORM_IMPLEMENTS_IO
+#define PLATFORM_IMPLEMENTS_IO 0
+#endif
+
+#ifndef PLATFORM_IODISPATCHER_MODULE
+#define PLATFORM_IODISPATCHER_MODULE PREPROCESSOR_TO_STRING(PREPROCESSOR_JOIN(PLATFORM_HEADER_NAME, PlatformIoDispatcher))
+#endif
+
+#define UE_FILEIOSTORE_STATS_ENABLED (COUNTERSTRACE_ENABLED || CSV_PROFILER)
+
+struct FFileIoStoreCompressionContext;
+
+struct FFileIoStoreContainerFilePartition
+{
+	uint64 FileHandle = 0;
+	uint64 FileSize = 0;
+	uint32 ContainerFileIndex = 0;
+	FString FilePath;
+	TUniquePtr<IMappedFileHandle> MappedFileHandle;
+};
+
+struct FFileIoStoreContainerFile
+{
+	uint64 PartitionSize = 0;
+	uint64 CompressionBlockSize = 0;
+	TArray<FName> CompressionMethods;
+	TArray<FIoStoreTocCompressedBlockEntry> CompressionBlocks;
+	FString FilePath;
+	FGuid EncryptionKeyGuid;
+	FAES::FAESKey EncryptionKey;
+	EIoContainerFlags ContainerFlags;
+	TArray<FSHAHash> BlockSignatureHashes;
+	TArray<FFileIoStoreContainerFilePartition> Partitions;
+	uint32 ContainerInstanceId = 0;
+
+	void GetPartitionFileHandleAndOffset(uint64 TocOffset, uint64& OutFileHandle, uint64& OutOffset) const
+	{
+		int32 PartitionIndex = int32(TocOffset / PartitionSize);
+		const FFileIoStoreContainerFilePartition& Partition = Partitions[PartitionIndex];
+		OutFileHandle = Partition.FileHandle;
+		OutOffset = TocOffset % PartitionSize;
+	}
+};
+
+struct FFileIoStoreBuffer
+{
+	FFileIoStoreBuffer* Next = nullptr;
+	uint8* Memory = nullptr;
+};
+
+struct FFileIoStoreBlockKey
+{
+	union
+	{
+		struct
+		{
+			uint32 FileIndex;
+			uint32 BlockIndex;
+		};
+		uint64 Hash;
+	};
+
+
+	friend bool operator==(const FFileIoStoreBlockKey& A, const FFileIoStoreBlockKey& B)
+	{
+		return A.Hash == B.Hash;
+	}
+
+	friend uint32 GetTypeHash(const FFileIoStoreBlockKey& Key)
+	{
+		return GetTypeHash(Key.Hash);
+	}
+};
+
+struct FFileIoStoreBlockScatter
+{
+	struct FFileIoStoreResolvedRequest* Request = nullptr;
+	uint64 DstOffset = 0;
+	uint64 SrcOffset = 0;
+	uint64 Size = 0;
+};
+
+struct FFileIoStoreCompressedBlock
+{
+	FFileIoStoreCompressedBlock* Next = nullptr;
+	FFileIoStoreBlockKey Key;
+	FName CompressionMethod;
+	uint64 RawOffset = uint64(-1);
+	uint32 UncompressedSize = uint32(-1);
+	uint32 CompressedSize = uint32(-1);
+	uint32 RawSize = uint32(-1);
+	uint32 RefCount = 0;
+	uint32 UnfinishedRawBlocksCount = 0;
+	TArray<struct FFileIoStoreReadRequest*, TInlineAllocator<2>> RawBlocks;
+	TArray<FFileIoStoreBlockScatter, TInlineAllocator<2>> ScatterList;
+	FFileIoStoreCompressionContext* CompressionContext = nullptr;
+	uint8* CompressedDataBuffer = nullptr;
+	FAES::FAESKey EncryptionKey;
+	const FSHAHash* SignatureHash = nullptr;
+	bool bFailed = false;
+	bool bCancelled = false;
+};
+
+struct FFileIoStoreReadRequest
+{
+	enum EQueueStatus
+	{
+		QueueStatus_NotInQueue,
+		QueueStatus_InQueue,
+		QueueStatus_Started
+	};
+
+	FFileIoStoreReadRequest()
+		: Sequence(NextSequence++)
+		, CreationTime(FPlatformTime::Cycles64())
+	{
+	}
+	FFileIoStoreReadRequest* Next = nullptr;
+	FFileIoStoreReadRequest* Previous = nullptr;
+	uint64 FileHandle = uint64(-1);
+	uint64 Offset = uint64(-1);
+	uint64 Size = uint64(-1);
+	FFileIoStoreBlockKey Key;
+	FFileIoStoreBuffer* Buffer = nullptr;
+	uint32 RefCount = 0;
+	uint32 BufferRefCount = 0;
+	TArray<FFileIoStoreCompressedBlock*, TInlineAllocator<8>> CompressedBlocks;
+	const uint32 Sequence;
+	int32 Priority = 0;
+	uint64 CreationTime;	// Potentially used to circuit break request ordering optimizations when outstanding requests have been delayed too long
+	FFileIoStoreBlockScatter ImmediateScatter;
+	uint32 BytesUsed = 0;
+	bool bIsCustomRequest = false;
+	bool bFailed = false;
+	bool bCancelled = false;
+	EQueueStatus QueueStatus = QueueStatus_NotInQueue;
+
+#if DO_CHECK
+	// For debug checks that we are in the correct owning list for our intrusive next/previous pointers
+	uint32 ListCookie = 0;
+#endif
+
+private:
+	static uint32 NextSequence;
+};
+
+#define CHECK_IO_STORE_READ_REQUEST_LIST_MEMBERSHIP (DO_CHECK && 0)
+// Iterator class for traversing and emptying a list FFileIoStoreReadRequestList at the same time
+class FFileIoStoreReadRequestListStealingIterator
+{
+public:
+	FFileIoStoreReadRequestListStealingIterator(const FFileIoStoreReadRequestListStealingIterator&) = delete;
+	FFileIoStoreReadRequestListStealingIterator& operator=(const FFileIoStoreReadRequestListStealingIterator&) = delete;
+
+	FFileIoStoreReadRequestListStealingIterator(FFileIoStoreReadRequestListStealingIterator&& Other)
+	{
+		Current = Other.Current;
+		Next = Other.Next;
+
+		Other.Current = Other.Next = nullptr;
+	}
+
+	FFileIoStoreReadRequestListStealingIterator& operator=(FFileIoStoreReadRequestListStealingIterator&& Other)
+	{
+		Current = Other.Current;
+		Next = Other.Next;
+
+		Other.Current = Other.Next = nullptr;
+
+		return *this;
+	}
+
+	void operator++()
+	{
+		AdvanceTo(Next);
+	}
+
+	FFileIoStoreReadRequest* operator*()
+	{
+		return Current;
+	}
+
+	FFileIoStoreReadRequest* operator->()
+	{
+		return Current;
+	}
+
+	explicit operator bool() const
+	{
+		return Current != nullptr;
+	}
+
+private:
+	friend class FFileIoStoreReadRequestList; // Only the list can construct us
+
+	FFileIoStoreReadRequestListStealingIterator(FFileIoStoreReadRequest* InHead)
+	{
+#if CHECK_IO_STORE_READ_REQUEST_LIST_MEMBERSHIP
+		// None of these nodes are members of the list any more
+		for (FFileIoStoreReadRequest* Cursor = InHead; Cursor; Cursor = Cursor->Next)
+		{
+			Cursor->ListCookie = 0;
+		}
+#endif
+
+		AdvanceTo(InHead);
+	}
+
+	void AdvanceTo(FFileIoStoreReadRequest* NewCurrent)
+	{
+		Current = NewCurrent;
+		if (Current)
+		{
+			// Copy off the next ptr and remove Current from its list so it can safely be added to another list
+			Next = Current->Next;
+			Current->Next = nullptr;
+		}
+		else
+		{
+			Next = nullptr;
+		}
+	}
+
+	FFileIoStoreReadRequest* Current = nullptr;
+	FFileIoStoreReadRequest* Next = nullptr;;
+};
+
+class FFileIoStoreReadRequestListIterator
+{
+public:
+	FFileIoStoreReadRequestListIterator(const FFileIoStoreReadRequestListIterator&) = default;
+	FFileIoStoreReadRequestListIterator& operator=(const FFileIoStoreReadRequestListIterator&) = default;
+
+	void operator++()
+	{
+		checkf(Current->Next == Next, TEXT("FFileIoStoreReadRequestListIterator cached Next pointer != Current->Next - was list modified during iteration?"));
+		AdvanceTo(Next);
+	}
+
+	FFileIoStoreReadRequest* operator*()
+	{
+		return Current;
+	}
+
+	FFileIoStoreReadRequest* operator->()
+	{
+		return Current;
+	}
+
+	explicit operator bool() const
+	{
+		return Current != nullptr;
+	}
+
+	bool operator!=(const FFileIoStoreReadRequestListIterator& Other) const
+	{
+		return Current != Other.Current;
+	}
+
+private:
+	friend class FFileIoStoreReadRequestList; // Only the list can construct us
+
+	FFileIoStoreReadRequestListIterator(FFileIoStoreReadRequest* InHead)
+	{
+		AdvanceTo(InHead);
+	}
+
+	void AdvanceTo(FFileIoStoreReadRequest* NewCurrent)
+	{
+		Current = NewCurrent;
+
+		if (Current)
+		{
+			// Copy off the next ptr to check for edits made to this list during iteration
+			Next = Current->Next;
+		}
+		else
+		{
+			Next = nullptr;
+		}
+	}
+
+	FFileIoStoreReadRequest* Current = nullptr;
+	FFileIoStoreReadRequest* Next = nullptr;
+};
+
+
+// Wrapper for doubly-linked intrusive list of FFileIoStoreReadRequest
+
+class FFileIoStoreReadRequestList
+{
+public:
+	FFileIoStoreReadRequestList()
+#if CHECK_IO_STORE_READ_REQUEST_LIST_MEMBERSHIP
+		: ListCookie(++NextListCookie)
+#endif
+	{
+	}
+
+	// Owns the Next/Previous pointers of the FFileIoStoreReadRequests it contains so it can't be copied
+	FFileIoStoreReadRequestList(const FFileIoStoreReadRequestList& Other) = delete;
+	FFileIoStoreReadRequestList& operator=(const FFileIoStoreReadRequestList& Other) = delete;
+
+	FFileIoStoreReadRequestList(FFileIoStoreReadRequestList&& Other)
+	{
+		Head = Other.Head;
+		Tail = Other.Tail;
+		Other.Head = Other.Tail = nullptr;
+
+#if CHECK_IO_STORE_READ_REQUEST_LIST_MEMBERSHIP
+		for (FFileIoStoreReadRequest* Cursor = Head; Cursor; Cursor = Cursor->Next)
+		{
+			Cursor->ListCookie = ListCookie;
+		}
+#endif
+	}
+
+	FFileIoStoreReadRequestList& operator=(FFileIoStoreReadRequestList&& Other)
+	{
+		check(!Head && !Tail);
+
+		Head = Other.Head;
+		Tail = Other.Tail;
+		Other.Head = Other.Tail = nullptr;
+
+#if CHECK_IO_STORE_READ_REQUEST_LIST_MEMBERSHIP
+		for (FFileIoStoreReadRequest* Cursor = Head; Cursor; Cursor = Cursor->Next)
+		{
+			Cursor->ListCookie = ListCookie;
+		}
+#endif
+
+		return *this;
+	}
+
+#if CHECK_IO_STORE_READ_REQUEST_LIST_MEMBERSHIP
+	~FFileIoStoreReadRequestList()
+	{
+		// If anything is left in this list it should still think we own it
+		for (FFileIoStoreReadRequest* Cursor = Head; Cursor; Cursor = Cursor->Next)
+		{
+			check(Cursor->ListCookie == ListCookie);
+		}
+	}
+#endif
+
+	bool IsEmpty() const
+	{
+		return Head == nullptr;
+	}
+
+	// Steal the whole list for iteration and moving the contents to other lists
+	FFileIoStoreReadRequestListStealingIterator Steal()
+	{
+#if CHECK_IO_STORE_READ_REQUEST_LIST_MEMBERSHIP
+		for (FFileIoStoreReadRequest* Cursor = Head; Cursor; Cursor = Cursor->Next)
+		{
+			check(Cursor->ListCookie == ListCookie);
+		}
+#endif
+
+		FFileIoStoreReadRequest* OldHead = Head;
+		Clear();
+		return FFileIoStoreReadRequestListStealingIterator(OldHead);
+	}
+
+	// Non-stealing iterator
+	FFileIoStoreReadRequestListIterator CreateIterator() const
+	{
+		return FFileIoStoreReadRequestListIterator(Head);
+	}
+
+	FFileIoStoreReadRequestListIterator begin() const {
+		return CreateIterator();
+	}
+	FFileIoStoreReadRequestListIterator end() const {
+		return FFileIoStoreReadRequestListIterator(nullptr);
+	}
+
+	FFileIoStoreReadRequest* PeekHead() const
+	{
+		return Head;
+	}
+
+	void Add(FFileIoStoreReadRequest* Request)
+	{
+#if CHECK_IO_STORE_READ_REQUEST_LIST_MEMBERSHIP
+		check(Request->ListCookie == 0);
+		Request->ListCookie = ListCookie;
+#endif
+
+		if (Tail)
+		{
+			Tail->Next = Request;
+			Request->Previous = Tail;
+		}
+		else
+		{
+			Head = Request;
+			Request->Previous = nullptr;
+		}
+		Tail = Request;
+		Request->Next = nullptr;
+	}
+
+
+	// Remove all FFileIoStoreReadRequests from List and add them to this list
+	void AppendSteal(FFileIoStoreReadRequestList& List)
+	{
+		if (List.Head)
+		{
+			FFileIoStoreReadRequest *ListHead = List.Head, *ListTail = List.Tail;
+			List.Clear();
+			AppendSteal(ListHead, ListTail);
+		}
+	}
+
+	void Remove(FFileIoStoreReadRequest* Request)
+	{
+#if CHECK_IO_STORE_READ_REQUEST_LIST_MEMBERSHIP
+		check(Request);
+		check(Request->ListCookie == ListCookie);
+		Request->ListCookie = 0;
+#endif
+
+		if (Head == Request && Tail == Request)
+		{
+			check(Request->Next == nullptr);
+			check(Request->Previous == nullptr);
+
+			Head = Tail = nullptr;
+		}
+		else if (Head == Request)
+		{
+			check(Request->Previous == nullptr);
+
+			Head = Request->Next;
+			Head->Previous = nullptr;
+			Request->Next = nullptr;
+		}
+		else if (Tail == Request)
+		{
+			check(Request->Next == nullptr);
+
+			Tail = Request->Previous;
+			Tail->Next = nullptr;
+			Request->Previous = nullptr;
+		}
+		else
+		{
+			check(Request->Next != nullptr && Request->Previous != nullptr); // Neither head nor tail should mean both links are live
+
+			Request->Next->Previous = Request->Previous;
+			Request->Previous->Next = Request->Next;
+
+			Request->Next = Request->Previous = nullptr;
+		}
+	}
+
+	void Clear()
+	{
+#if CHECK_IO_STORE_READ_REQUEST_LIST_MEMBERSHIP
+		for (FFileIoStoreReadRequest* Cursor = Head; Cursor; Cursor = Cursor->Next)
+		{
+			check(Cursor->ListCookie == ListCookie);
+			Cursor->ListCookie = 0;
+		}
+#endif
+		Head = Tail = nullptr;
+	}
+
+private:
+	FFileIoStoreReadRequest* Head = nullptr;
+	FFileIoStoreReadRequest* Tail = nullptr;
+
+#if CHECK_IO_STORE_READ_REQUEST_LIST_MEMBERSHIP
+	uint32 ListCookie;
+	static uint32 NextListCookie;
+#endif
+	
+	void AppendSteal(FFileIoStoreReadRequest* ListHead, FFileIoStoreReadRequest* ListTail)
+	{
+		check(ListHead);
+		check(ListTail);
+		check(!ListTail->Next);
+		check(!ListHead->Previous);
+		check(ListTail == ListHead || ListTail->Previous != nullptr);
+		check(ListTail == ListHead || ListHead->Next != nullptr);
+
+#if CHECK_IO_STORE_READ_REQUEST_LIST_MEMBERSHIP
+		for (FFileIoStoreReadRequest* Cursor = ListHead; Cursor; Cursor = Cursor->Next)
+		{
+			check(Cursor->ListCookie == 0);
+			Cursor->ListCookie = ListCookie;
+		}
+#endif
+
+		if (Tail)
+		{
+			Tail->Next = ListHead;
+			ListHead->Previous = Tail;
+		}
+		else
+		{
+			Head = ListHead;
+			ListHead->Previous = nullptr;
+		}
+		Tail = ListTail;
+	}
+};
+
+class FFileIoStoreStats;
+
+class FFileIoStoreBufferAllocator
+{
+public:
+	FFileIoStoreBufferAllocator(FFileIoStoreStats& InStats)
+		: Stats(InStats)
+	{
+	}
+
+	void Initialize(uint64 MemorySize, uint64 BufferSize, uint32 BufferAlignment);
+	FFileIoStoreBuffer* AllocBuffer();
+	void FreeBuffer(FFileIoStoreBuffer* Buffer);
+	uint64 GetBufferSize() const { return BufferSize; }
+
+private:
+	FFileIoStoreStats& Stats;
+	uint64 BufferSize = 0;
+	uint8* BufferMemory = nullptr;
+	FCriticalSection BuffersCritical;
+	FFileIoStoreBuffer* FirstFreeBuffer = nullptr;
+};
+
+class FFileIoStoreBlockCache
+{
+public:
+	FFileIoStoreBlockCache(FFileIoStoreStats& Stats);
+	~FFileIoStoreBlockCache();
+
+	void Initialize(uint64 CacheMemorySize, uint64 ReadBufferSize);
+	bool Read(FFileIoStoreReadRequest* Block);
+	void Store(const FFileIoStoreReadRequest* Block);
+
+private:
+	struct FCachedBlock
+	{
+		FCachedBlock* LruPrev = nullptr;
+		FCachedBlock* LruNext = nullptr;
+		uint64 Key = 0;
+		uint8* Buffer = nullptr;
+	};
+
+	FFileIoStoreStats& Stats;
+	uint8* CacheMemory = nullptr;
+	TMap<uint64, FCachedBlock*> CachedBlocks;
+	FCachedBlock CacheLruHead;
+	FCachedBlock CacheLruTail;
+	uint64 ReadBufferSize = 0;
+};
+
+struct FFileIoStoreReadRequestSortKey
+{
+	uint64 Offset = 0;
+	uint64 Handle = 0;
+	int32 Priority = 0;
+
+	FFileIoStoreReadRequestSortKey() {}
+	FFileIoStoreReadRequestSortKey(FFileIoStoreReadRequest* Request)
+		: Offset(Request->Offset), Handle(Request->FileHandle), Priority(Request->Priority)
+	{
+	}
+};
+
+// Stores FFileIoStoreReadRequest sorted by file handle & offset with a parallel list sorted by insertion order 
+class FFileIoStoreOffsetSortedRequestQueue
+{
+public:
+	FFileIoStoreOffsetSortedRequestQueue(int32 InPriority);
+	FFileIoStoreOffsetSortedRequestQueue(const FFileIoStoreOffsetSortedRequestQueue&) = delete;
+	FFileIoStoreOffsetSortedRequestQueue(FFileIoStoreOffsetSortedRequestQueue&&) = default;
+	FFileIoStoreOffsetSortedRequestQueue& operator=(const FFileIoStoreOffsetSortedRequestQueue&) = delete;
+	FFileIoStoreOffsetSortedRequestQueue& operator=(FFileIoStoreOffsetSortedRequestQueue&&) = default;
+
+	int32 GetPriority() const { return Priority; }
+	bool IsEmpty() const { return Requests.Num() == 0; }
+	
+	// Remove all requests from this container and return them, for switching to a different queue scheme
+	TArray<FFileIoStoreReadRequest*> StealRequests();
+	// Remove all requests whose priority has been changed to something other than the Priority of this queue
+	TArray<FFileIoStoreReadRequest*> RemoveMisprioritizedRequests();
+
+	FFileIoStoreReadRequest* Peek(FFileIoStoreReadRequestSortKey LastSortKey);
+	FFileIoStoreReadRequest* Pop(FFileIoStoreReadRequestSortKey LastSortKey);
+	void Push(FFileIoStoreReadRequest* Request);
+	void CancelRequestsWithFileHandle(uint64 FileHandle);
+
+private:
+	int32 Priority;
+	int32 PeekRequestIndex = INDEX_NONE;
+	
+	// Requests sorted by file handle & offset
+	TArray<FFileIoStoreReadRequest*> Requests;
+	
+	// Requests sorted by insertion order 
+	// We store this on the heap in case we get moved, FFileIoStoreReadRequest keeps pointers to FFileIoStoreReadRequestList for debugging
+	FFileIoStoreReadRequestList RequestsBySequence;
+
+	FFileIoStoreReadRequest* GetNextInternal(FFileIoStoreReadRequestSortKey LastSortKey, bool bPop);
+
+	static FFileIoStoreReadRequestSortKey RequestSortProjection(FFileIoStoreReadRequest* Request) { return FFileIoStoreReadRequestSortKey(Request); }
+	static bool RequestSortPredicate(const FFileIoStoreReadRequestSortKey& A, const FFileIoStoreReadRequestSortKey& B);
+};
+
+class FFileIoStoreRequestQueue
+{
+public:
+	FFileIoStoreReadRequest* Peek();
+	FFileIoStoreReadRequest* Pop();
+	void Push(FFileIoStoreReadRequest& Request);	// Takes ownership of Request and rewrites its intrustive linked list pointers
+	void Push(FFileIoStoreReadRequestList& Requests); // Consumes the request list and overwrites all intrustive linked list pointers
+	void UpdateOrder();
+	void Lock();
+	void Unlock();
+	void CancelRequestsWithFileHandle(uint64 FileHandle);
+
+private:
+	static bool QueueSortFunc(const FFileIoStoreReadRequest& A, const FFileIoStoreReadRequest& B)
+	{
+		if (A.Priority == B.Priority)
+		{
+			return A.Sequence < B.Sequence;
+		}
+		return A.Priority > B.Priority;
+	}
+	void UpdateSortRequestsByOffset(); // Check if we need to switch sorting schemes based on the CVar
+	void PushToPriorityQueues(FFileIoStoreReadRequest* Request);
+	static int32 QueuePriorityProjection(const FFileIoStoreOffsetSortedRequestQueue& A) { return A.GetPriority(); }
+
+	bool bSortRequestsByOffset = false; // Cached value of CVar controlling whether we use Heap or SortedPriorityQueues
+	
+	// Heap sorted by request order
+	TArray<FFileIoStoreReadRequest*> Heap;
+	FCriticalSection CriticalSection;
+
+	// Queues sorted by increasing priority
+	TArray<FFileIoStoreOffsetSortedRequestQueue> SortedPriorityQueues; 
+	// The last offset, file handle and priority we popped so that we can pop the closest forward read for the next IO operation
+	FFileIoStoreReadRequestSortKey LastSortKey;
+
+#if !UE_BUILD_SHIPPING
+	TMap<int32, uint32> RequestPriorityCounts;
+#endif
+};
+
+template <typename T, uint16 SlabSize = 4096>
+class TIoDispatcherSingleThreadedSlabAllocator
+{
+public:
+	TIoDispatcherSingleThreadedSlabAllocator()
+	{
+		CurrentSlab = new FSlab();
+	}
+
+	~TIoDispatcherSingleThreadedSlabAllocator()
+	{
+		check(CurrentSlab->Allocated == CurrentSlab->Freed);
+		delete CurrentSlab;
+	}
+
+	template <typename... ArgsType>
+	T* Construct(ArgsType&&... Args)
+	{
+		return new(Alloc()) T(Forward<ArgsType>(Args)...);
+	}
+
+	void Destroy(T* Ptr)
+	{
+		Ptr->~T();
+		Free(Ptr);
+	}
+
+private:
+	struct FSlab;
+
+	struct FElement
+	{
+		TTypeCompatibleBytes<T> Data;
+		FSlab* Slab = nullptr;
+	};
+
+	struct FSlab
+	{
+		uint16 Allocated = 0;
+		uint16 Freed = 0;
+		FElement Elements[SlabSize];
+	};
+
+	T* Alloc()
+	{
+		uint16 ElementIndex = CurrentSlab->Allocated++;
+		check(ElementIndex < SlabSize);
+		FElement* Element = CurrentSlab->Elements + ElementIndex;
+		Element->Slab = CurrentSlab;
+		if (CurrentSlab->Allocated == SlabSize)
+		{
+			//TRACE_CPUPROFILER_EVENT_SCOPE(AllocSlab);
+			CurrentSlab = new FSlab();
+		}
+		return Element->Data.GetTypedPtr();
+	}
+
+	void Free(T* Ptr)
+	{
+		FElement* Element = reinterpret_cast<FElement*>(Ptr);
+		FSlab* Slab = Element->Slab;
+		if (++Slab->Freed == SlabSize)
+		{
+			//TRACE_CPUPROFILER_EVENT_SCOPE(FreeSlab);
+			check(Slab->Freed == Slab->Allocated);
+			delete Slab;
+		}
+	}
+
+	FSlab* CurrentSlab = nullptr;
+};
+
+struct FFileIoStoreReadRequestLink
+{
+	FFileIoStoreReadRequestLink(FFileIoStoreReadRequest& InReadRequest)
+		: ReadRequest(InReadRequest)
+	{
+	}
+
+	FFileIoStoreReadRequestLink* Next = nullptr;
+	FFileIoStoreReadRequest& ReadRequest;
+};
+
+class FFileIoStoreRequestAllocator
+{
+public:
+	int64 GetLiveReadRequestsCount() const
+	{
+		return LiveReadRequestsCount;
+	}
+
+
+	FFileIoStoreResolvedRequest* AllocResolvedRequest(
+		FIoRequestImpl& InDispatcherRequest,
+		const FFileIoStoreContainerFile& InContainerFile,
+		uint64 InResolvedOffset,
+		uint64 InResolvedSize)
+	{
+		return ResolvedRequestAllocator.Construct(
+			InDispatcherRequest,
+			InContainerFile,
+			InResolvedOffset,
+			InResolvedSize);
+	}
+
+	void Free(FFileIoStoreResolvedRequest* ResolvedRequest)
+	{
+		ResolvedRequestAllocator.Destroy(ResolvedRequest);
+	}
+
+	FFileIoStoreReadRequest* AllocReadRequest()
+	{
+		//TRACE_CPUPROFILER_EVENT_SCOPE(AllocReadRequest);
+		++LiveReadRequestsCount;
+		return ReadRequestAllocator.Construct();
+	}
+
+	void Free(FFileIoStoreReadRequest* ReadRequest)
+	{
+		//TRACE_CPUPROFILER_EVENT_SCOPE(FreeReadRequest);
+		ReadRequestAllocator.Destroy(ReadRequest);
+		--LiveReadRequestsCount;
+	}
+
+	FFileIoStoreCompressedBlock* AllocCompressedBlock()
+	{
+		//TRACE_CPUPROFILER_EVENT_SCOPE(AllocCompressedBlock);
+		return CompressedBlockAllocator.Construct();
+	}
+
+	void Free(FFileIoStoreCompressedBlock* CompressedBlock)
+	{
+		//TRACE_CPUPROFILER_EVENT_SCOPE(FreeCompressedBlock);
+		CompressedBlockAllocator.Destroy(CompressedBlock);
+	}
+
+	FFileIoStoreReadRequestLink* AllocRequestLink(FFileIoStoreReadRequest* ReadRequest)
+	{
+		check(ReadRequest);
+		return RequestLinkAllocator.Construct(*ReadRequest);
+	}
+
+	void Free(FFileIoStoreReadRequestLink* RequestLink)
+	{
+		RequestLinkAllocator.Destroy(RequestLink);
+	}
+
+private:
+	TIoDispatcherSingleThreadedSlabAllocator<FFileIoStoreResolvedRequest> ResolvedRequestAllocator;
+	TIoDispatcherSingleThreadedSlabAllocator<FFileIoStoreReadRequest> ReadRequestAllocator;
+	TIoDispatcherSingleThreadedSlabAllocator<FFileIoStoreCompressedBlock> CompressedBlockAllocator;
+	TIoDispatcherSingleThreadedSlabAllocator<FFileIoStoreReadRequestLink> RequestLinkAllocator;
+	int64 LiveReadRequestsCount = 0;
+};
+
+struct FFileIoStoreResolvedRequest
+{
+public:
+	FFileIoStoreResolvedRequest(
+		FIoRequestImpl& InDispatcherRequest,
+		const FFileIoStoreContainerFile& InContainerFile,
+		uint64 InResolvedOffset,
+		uint64 InResolvedSize);
+
+	const FFileIoStoreContainerFile& GetContainerFile() const
+	{
+		return ContainerFile;
+	}
+
+	uint64 GetResolvedOffset() const
+	{
+		return ResolvedOffset;
+	}
+
+	uint64 GetResolvedSize() const
+	{
+		return ResolvedSize;
+	}
+
+	int32 GetPriority() const
+	{
+		check(DispatcherRequest);
+		return DispatcherRequest->Priority;
+	}
+
+	bool HasBuffer() const
+	{
+		check(DispatcherRequest);
+		return DispatcherRequest->HasBuffer();
+	}
+
+	void CreateBuffer(uint64 Size)
+	{
+		check(DispatcherRequest);
+		DispatcherRequest->CreateBuffer(Size);
+	}
+
+	FIoBuffer& GetBuffer()
+	{
+		check(DispatcherRequest);
+		return DispatcherRequest->GetBuffer();
+	}
+
+	void AddReadRequestLink(FFileIoStoreReadRequestLink* ReadRequestLink);
+
+private:
+	FIoRequestImpl* DispatcherRequest = nullptr;
+	const FFileIoStoreContainerFile& ContainerFile;
+	FFileIoStoreReadRequestLink* ReadRequestsHead = nullptr;
+	FFileIoStoreReadRequestLink* ReadRequestsTail = nullptr;
+	const uint64 ResolvedOffset;
+	const uint64 ResolvedSize;
+	uint32 UnfinishedReadsCount = 0;
+	bool bFailed = false;
+	bool bCancelled = false;
+
+	friend class FFileIoStore;
+	friend class FFileIoStoreRequestTracker;
+};
+
+
+// Wrapper for sending stats to both insights and csv profiling from any platform-specific dispatcher.
+#if UE_FILEIOSTORE_STATS_ENABLED
+class FFileIoStoreStats
+{
+public:
+	FFileIoStoreStats();
+	~FFileIoStoreStats();
+
+	// Called by the backend when underlying file system reads start
+	void OnFilesystemReadStarted(const FFileIoStoreReadRequest* Request);
+	void OnFilesystemReadsStarted(const FFileIoStoreReadRequestList& Requests);
+	// Called by the backend when underlying filesystem reads complete, possibly already decompressed on some systems
+	void OnFilesystemReadCompleted(const FFileIoStoreReadRequest* Request);
+	void OnFilesystemReadsCompleted(const FFileIoStoreReadRequestList& CompletedRequests);
+
+private:
+	friend class FFileIoStore;
+	friend class FFileIoStoreBlockCache;
+	friend class FFileIoStoreReader;
+	friend class FFileIoStoreBufferAllocator;
+
+#if COUNTERSTRACE_ENABLED
+	FCountersTrace::FCounterInt QueuedReadRequestsSizeCounter;
+	FCountersTrace::FCounterInt CompletedReadRequestsSizeCounter;
+	FCountersTrace::FCounterInt QueuedCompressedSizeCounter;
+	FCountersTrace::FCounterInt QueuedUncompressedSizeCounter;
+	FCountersTrace::FCounterInt CompletedCompressedSizeCounter;
+	FCountersTrace::FCounterInt CompletedUncompressedSizeCounter;
+	FCountersTrace::FCounterInt FileSystemSeeksTotalDistanceCounter;
+	FCountersTrace::FCounterInt FileSystemSeeksForwardCountCounter;
+	FCountersTrace::FCounterInt FileSystemSeeksBackwardCountCounter;
+	FCountersTrace::FCounterInt FileSystemSeeksChangeHandleCountCounter;
+	FCountersTrace::FCounterInt FileSystemCompletedRequestsSizeCounter;
+	FCountersTrace::FCounterInt BlockCacheStoredSizeCounter;
+	FCountersTrace::FCounterInt BlockCacheHitSizeCounter;
+	FCountersTrace::FCounterInt BlockCacheMissedSizeCounter;
+	FCountersTrace::FCounterInt ScatteredSizeCounter;
+	FCountersTrace::FCounterInt TocMemoryCounter;
+	FCountersTrace::FCounterInt AvailableBuffersCounter;
+#endif
+
+#if CSV_PROFILER
+	uint64 QueuedFilesystemReadBytes = 0;
+	uint64 QueuedFilesystemReads = 0;
+
+	uint64 QueuedUncompressBytesIn = 0;
+	uint64 QueuedUncompressBytesOut = 0;
+	uint64 QueuedUncompressBlocks = 0;
+#endif
+
+	// Used for seek tracking 
+	uint64 LastHandle = uint64(-1);
+	uint64 LastOffset = uint64(-1);
+
+	FTSTicker::FDelegateHandle TickerHandle;
+
+	float BytesToApproxMB(uint64 Bytes)
+	{
+		return float(double(Bytes) / 1024.0 / 1024.0);
+	}
+
+	float BytesToApproxKB(uint64 Bytes)
+	{
+		return float(double(Bytes) / 1024.0);
+	}
+
+	bool CsvTick(float DeltaTime);
+
+	void OnReadRequestsQueued(const FFileIoStoreReadRequestList& Requests);
+	void OnReadRequestsCompleted(const FFileIoStoreReadRequestList& Requests);
+	void OnDecompressQueued(const FFileIoStoreCompressedBlock* CompressedBlock);
+	void OnDecompressComplete(const FFileIoStoreCompressedBlock* CompressedBlock);
+
+	// Bytes were copied into the buffer provided to users of the io system
+	void OnBytesScattered(int64 BytesScattered);
+	// Record stats for block cache hit/miss rate & data throughput
+	void OnBlockCacheStore(uint64 NumBytes);
+	void OnBlockCacheHit(uint64 NumBytes);
+	void OnBlockCacheMiss(uint64 NumBytes);
+
+	// A read was started without seeking
+	void OnSequentialRead();
+	// A read was started and was either a forward or reverse seek
+	void OnSeek(uint64 LastOffset, uint64 NewOffset);
+	// A read was started from a different file handle from the last one we used
+	void OnHandleChangeSeek();
+
+	void OnTocMounted(uint64 AllocatedSize);
+	void OnTocUnmounted(uint64 AllocatedSize);
+
+	void OnBufferReleased();
+	void OnBufferAllocated();
+};
+#else
+class FFileIoStoreStats
+{
+public:
+	void OnFilesystemReadStarted(const FFileIoStoreReadRequest* Request) {}
+	void OnFilesystemReadsStarted(const FFileIoStoreReadRequestList& Requests) {}
+	void OnFilesystemReadCompleted(const FFileIoStoreReadRequest* Request) {}
+	void OnFilesystemReadsCompleted(const FFileIoStoreReadRequestList& CompletedRequests) {}
+
+private:
+	friend class FFileIoStore;
+	friend class FFileIoStoreBlockCache;
+	friend class FFileIoStoreReader;
+	friend class FFileIoStoreBufferAllocator;
+
+	void OnReadRequestsQueued(const FFileIoStoreReadRequestList& Requests) {}
+	void OnReadRequestsCompleted(const FFileIoStoreReadRequestList& Requests) {}
+	void OnDecompressQueued(const FFileIoStoreCompressedBlock* CompressedBlock) {}
+	void OnDecompressComplete(const FFileIoStoreCompressedBlock* CompressedBlock) {}
+	void OnBytesScattered(int64 BytesScattered) {}
+	void OnBlockCacheStore(uint64 NumBytes) {}
+	void OnBlockCacheHit(uint64 NumBytes) {}
+	void OnBlockCacheMiss(uint64 NumBytes) {}
+	void OnTocMounted(uint64 AllocatedSize) {}
+	void OnTocUnmounted(uint64 AllocatedSize) {}
+	void OnBufferReleased() {}
+	void OnBufferAllocated() {}
+};
+#endif
+
+struct FInitializePlatformFileIoStoreParams
+{
+	const FWakeUpIoDispatcherThreadDelegate* WakeUpDispatcherThreadDelegate = nullptr;
+	FFileIoStoreRequestAllocator* RequestAllocator = nullptr;
+	FFileIoStoreBufferAllocator* BufferAllocator = nullptr;
+	FFileIoStoreBlockCache* BlockCache = nullptr;
+	FFileIoStoreStats* Stats = nullptr;
+};
+
+class IPlatformFileIoStore
+{
+public:
+	virtual ~IPlatformFileIoStore() = default;
+	virtual void Initialize(const FInitializePlatformFileIoStoreParams& Params) = 0;
+
+	virtual bool OpenContainer(const TCHAR* ContainerFilePath, uint64& ContainerFileHandle, uint64& ContainerFileSize) = 0;
+	virtual void CloseContainer(uint64 ContainerFileHandle) = 0;
+
+	virtual bool CreateCustomRequests(FFileIoStoreResolvedRequest& ResolvedRequest, FFileIoStoreReadRequestList& OutRequests) = 0;
+	virtual bool StartRequests(FFileIoStoreRequestQueue& RequestQueue) = 0;
+	virtual void GetCompletedRequests(FFileIoStoreReadRequestList& OutRequests) = 0;
+
+	virtual void ServiceNotify() = 0;
+	virtual void ServiceWait() = 0;
+};
+
+class IPlatformFileIoStoreModule : public IModuleInterface
+{
+public:
+	virtual TUniquePtr<IPlatformFileIoStore> CreatePlatformFileIoStore() = 0;
+};
+
+#if PLATFORM_IMPLEMENTS_IO
+CORE_API TUniquePtr<IPlatformFileIoStore> CreatePlatformFileIoStore();
+#endif
+
diff --git a/Engine/Source/Runtime/PakFile/Private/GenericPlatformIoDispatcher.cpp b/Engine/Source/Runtime/PakFile/Private/GenericPlatformIoDispatcher.cpp
new file mode 100644
--- /dev/null
+++ b/Engine/Source/Runtime/PakFile/Private/GenericPlatformIoDispatcher.cpp
@@ -0,0 +1,144 @@
+// Copyright Epic Games, Inc. All Rights Reserved.
+
+#include "GenericPlatformIoDispatcher.h"
+#include "IoDispatcherFileBackend.h"
+#include "Async/AsyncFileHandle.h"
+#include "HAL/Event.h"
+#include "HAL/PlatformFileManager.h"
+#include "HAL/RunnableThread.h"
+#include "HAL/IConsoleManager.h"
+#include "Misc/ScopeLock.h"
+
+//PRAGMA_DISABLE_OPTIMIZATION
+
+
+TRACE_DECLARE_INT_COUNTER_EXTERN(IoDispatcherFileBackendSequentialReads);
+TRACE_DECLARE_INT_COUNTER_EXTERN(IoDispatcherFileBackendForwardSeeks);
+TRACE_DECLARE_INT_COUNTER_EXTERN(IoDispatcherFileBackendBackwardSeeks);
+TRACE_DECLARE_INT_COUNTER_EXTERN(IoDispatcherFileBackendSwitchContainerSeeks);
+TRACE_DECLARE_MEMORY_COUNTER_EXTERN(IoDispatcherFileBackendTotalSeekDistance);
+TRACE_DECLARE_INT_COUNTER_EXTERN(IoDispatcherFileBackendFileSystemRequests);
+TRACE_DECLARE_MEMORY_COUNTER_EXTERN(IoDispatcherFileBackendFileSystemTotalBytesRead);
+
+FGenericFileIoStoreEventQueue::FGenericFileIoStoreEventQueue()
+	: ServiceEvent(FPlatformProcess::GetSynchEventFromPool())
+{
+}
+
+FGenericFileIoStoreEventQueue::~FGenericFileIoStoreEventQueue()
+{
+	FPlatformProcess::ReturnSynchEventToPool(ServiceEvent);
+}
+
+void FGenericFileIoStoreEventQueue::ServiceNotify()
+{
+	ServiceEvent->Trigger();
+}
+
+void FGenericFileIoStoreEventQueue::ServiceWait()
+{
+	ServiceEvent->Wait();
+}
+
+FGenericFileIoStoreImpl::FGenericFileIoStoreImpl()
+{
+}
+
+FGenericFileIoStoreImpl::~FGenericFileIoStoreImpl()
+{
+}
+
+bool FGenericFileIoStoreImpl::OpenContainer(const TCHAR* ContainerFilePath, uint64& ContainerFileHandle, uint64& ContainerFileSize)
+{
+	IPlatformFile& Ipf = IPlatformFile::GetPlatformPhysical();
+	int64 FileSize = Ipf.FileSize(ContainerFilePath);
+	if (FileSize < 0)
+	{
+		return false;
+	}
+	IFileHandle* FileHandle = Ipf.OpenReadNoBuffering(ContainerFilePath);
+	if (!FileHandle)
+	{
+		return false;
+	}
+	ContainerFileHandle = reinterpret_cast<UPTRINT>(FileHandle);
+	ContainerFileSize = uint64(FileSize);
+	return true;
+}
+
+void FGenericFileIoStoreImpl::CloseContainer(uint64 ContainerFileHandle)
+{
+	check(ContainerFileHandle);
+	IFileHandle* FileHandle = reinterpret_cast<IFileHandle*>(ContainerFileHandle);
+	delete FileHandle;
+}
+
+bool FGenericFileIoStoreImpl::StartRequests(FFileIoStoreRequestQueue& RequestQueue)
+{
+	FFileIoStoreReadRequest* NextRequest = RequestQueue.Pop();
+	if (!NextRequest)
+	{
+		return false;
+	}
+
+	if (NextRequest->bCancelled)
+	{
+		{
+			FScopeLock _(&CompletedRequestsCritical);
+			CompletedRequests.Add(NextRequest);
+		}
+		WakeUpDispatcherThreadDelegate->Execute();
+		return true;
+	}
+
+	uint8* Dest;
+	check(!NextRequest->ImmediateScatter.Request);
+	NextRequest->Buffer = BufferAllocator->AllocBuffer();
+	if (!NextRequest->Buffer)
+	{
+		RequestQueue.Push(*NextRequest);
+		return false;
+	}
+	Dest = NextRequest->Buffer->Memory;
+
+	if (!BlockCache->Read(NextRequest))
+	{
+		IFileHandle* FileHandle = reinterpret_cast<IFileHandle*>(static_cast<UPTRINT>(NextRequest->FileHandle));
+		 
+		Stats->OnFilesystemReadStarted(NextRequest);
+		{
+			TRACE_CPUPROFILER_EVENT_SCOPE(ReadBlockFromFile);
+			NextRequest->bFailed = true;
+			int32 RetryCount = 0;
+			while (RetryCount++ < 10)
+			{
+				if (!FileHandle->Seek(NextRequest->Offset))
+				{
+					UE_LOG(LogIoDispatcher, Warning, TEXT("Failed seeking to offset %lld (Retries: %d)"), NextRequest->Offset, (RetryCount - 1));
+					continue;
+				}
+				if (!FileHandle->Read(Dest, NextRequest->Size))
+				{
+					UE_LOG(LogIoDispatcher, Warning, TEXT("Failed reading %lld bytes at offset %lld (Retries: %d)"), NextRequest->Size, NextRequest->Offset, (RetryCount - 1));
+					continue;
+				}
+				NextRequest->bFailed = false;
+				Stats->OnFilesystemReadCompleted(NextRequest);
+				BlockCache->Store(NextRequest);
+				break;
+			}
+		}
+	}
+	{
+		FScopeLock _(&CompletedRequestsCritical);
+		CompletedRequests.Add(NextRequest);
+	}
+	WakeUpDispatcherThreadDelegate->Execute();
+	return true;
+}
+
+void FGenericFileIoStoreImpl::GetCompletedRequests(FFileIoStoreReadRequestList& OutRequests)
+{
+	FScopeLock _(&CompletedRequestsCritical);
+	OutRequests.AppendSteal(CompletedRequests);
+}
diff --git a/Engine/Source/Runtime/PakFile/Private/GenericPlatformIoDispatcher.h b/Engine/Source/Runtime/PakFile/Private/GenericPlatformIoDispatcher.h
new file mode 100644
--- /dev/null
+++ b/Engine/Source/Runtime/PakFile/Private/GenericPlatformIoDispatcher.h
@@ -0,0 +1,68 @@
+// Copyright Epic Games, Inc. All Rights Reserved.
+
+#pragma once
+
+#include "Containers/List.h"
+#include "HAL/CriticalSection.h"
+#include "Templates/Tuple.h"
+#include "HAL/Runnable.h"
+#include "Containers/Map.h"
+#include "IO/IoDispatcher.h"
+#include "IoDispatcherFileBackendTypes.h"
+#include "ProfilingDebugging/CountersTrace.h"
+
+class FGenericFileIoStoreEventQueue
+{
+public:
+	FGenericFileIoStoreEventQueue();
+	~FGenericFileIoStoreEventQueue();
+	void ServiceNotify();
+	void ServiceWait();
+
+private:
+	FEvent* ServiceEvent = nullptr;
+};
+
+class FGenericFileIoStoreImpl : public IPlatformFileIoStore
+{
+public:
+	FGenericFileIoStoreImpl();
+	~FGenericFileIoStoreImpl();
+	void Initialize(const FInitializePlatformFileIoStoreParams& Params) override
+	{
+		WakeUpDispatcherThreadDelegate = Params.WakeUpDispatcherThreadDelegate;
+		BufferAllocator = Params.BufferAllocator;
+		BlockCache = Params.BlockCache;
+		Stats = Params.Stats;
+	}
+	bool OpenContainer(const TCHAR* ContainerFilePath, uint64& ContainerFileHandle, uint64& ContainerFileSize) override;
+	void CloseContainer(uint64 ContainerFileHandle) override;
+	bool CreateCustomRequests(FFileIoStoreResolvedRequest& ResolvedRequest, FFileIoStoreReadRequestList& OutRequests) override
+	{
+		return false;
+	}
+	bool StartRequests(FFileIoStoreRequestQueue& RequestQueue) override;
+	void GetCompletedRequests(FFileIoStoreReadRequestList& OutRequests) override;
+
+	virtual void ServiceNotify() override
+	{
+		EventQueue.ServiceNotify();
+	}
+	virtual void ServiceWait() override
+	{
+		EventQueue.ServiceWait();
+	}
+private:
+	const FWakeUpIoDispatcherThreadDelegate* WakeUpDispatcherThreadDelegate = nullptr;
+	FGenericFileIoStoreEventQueue EventQueue;
+	FFileIoStoreBufferAllocator* BufferAllocator = nullptr;
+	FFileIoStoreBlockCache* BlockCache = nullptr;
+	FFileIoStoreStats* Stats = nullptr;
+
+	FCriticalSection CompletedRequestsCritical;
+	FFileIoStoreReadRequestList CompletedRequests;
+#if COUNTERSTRACE_ENABLED
+	uint64 PreviousFileHandle = 0;
+#endif
+};
+
diff --git a/Engine/Source/Runtime/PakFile/Private/IPlatformFilePak.cpp b/Engine/Source/Runtime/PakFile/Private/IPlatformFilePak.cpp
--- a/Engine/Source/Runtime/PakFile/Private/IPlatformFilePak.cpp
+++ b/Engine/Source/Runtime/PakFile/Private/IPlatformFilePak.cpp
@@ -34,7 +34,7 @@
 #include "Misc/Fnv.h"
 #include "ProfilingDebugging/ScopedTimers.h"
 #include "Async/MappedFileHandle.h"
-#include "IO/IoDispatcherBackend.h"
+#include "IoDispatcherFileBackend.h"
 
 #include "ProfilingDebugging/LoadTimeTracker.h"
 #include "IO/IoContainerHeader.h"
diff --git a/Engine/Source/Runtime/PakFile/Private/IoDispatcherFileBackend.cpp b/Engine/Source/Runtime/PakFile/Private/IoDispatcherFileBackend.cpp
new file mode 100644
--- /dev/null
+++ b/Engine/Source/Runtime/PakFile/Private/IoDispatcherFileBackend.cpp
@@ -0,0 +1,2422 @@
+// Copyright Epic Games, Inc. All Rights Reserved.
+
+#include "IoDispatcherFileBackend.h"
+#include "Misc/CommandLine.h"
+#include "Misc/Parse.h"
+#include "Misc/ScopeRWLock.h"
+#include "Misc/StringBuilder.h"
+#include "ProfilingDebugging/CountersTrace.h"
+#include "HAL/PlatformFileManager.h"
+#include "GenericPlatform/GenericPlatformFile.h"
+#include "GenericPlatformIoDispatcher.h"
+#include "HAL/IConsoleManager.h"
+#include "Async/AsyncWork.h"
+#include "Async/MappedFileHandle.h"
+#include "HAL/RunnableThread.h"
+#include "Misc/ScopeLock.h"
+#include "Misc/Paths.h"
+#include "Algo/AllOf.h"
+#include "Algo/IsSorted.h"
+#include "Algo/MinElement.h"
+#include "Templates/Greater.h"
+#include "Containers/Ticker.h"
+#include "Modules/ModuleManager.h"
+#include "IO/IoContainerHeader.h"
+#include "Serialization/MemoryReader.h"
+
+//PRAGMA_DISABLE_OPTIMIZATION
+
+int32 GIoDispatcherBufferSizeKB = 256;
+static FAutoConsoleVariableRef CVar_IoDispatcherBufferSizeKB(
+	TEXT("s.IoDispatcherBufferSizeKB"),
+	GIoDispatcherBufferSizeKB,
+	TEXT("IoDispatcher read buffer size (in kilobytes).")
+);
+
+int32 GIoDispatcherBufferAlignment = 4096;
+static FAutoConsoleVariableRef CVar_IoDispatcherBufferAlignment(
+	TEXT("s.IoDispatcherBufferAlignment"),
+	GIoDispatcherBufferAlignment,
+	TEXT("IoDispatcher read buffer alignment.")
+);
+
+int32 GIoDispatcherBufferMemoryMB = 8;
+static FAutoConsoleVariableRef CVar_IoDispatcherBufferMemoryMB(
+	TEXT("s.IoDispatcherBufferMemoryMB"),
+	GIoDispatcherBufferMemoryMB,
+	TEXT("IoDispatcher buffer memory size (in megabytes).")
+);
+
+int32 GIoDispatcherDecompressionWorkerCount = 4;
+static FAutoConsoleVariableRef CVar_IoDispatcherDecompressionWorkerCount(
+	TEXT("s.IoDispatcherDecompressionWorkerCount"),
+	GIoDispatcherDecompressionWorkerCount,
+	TEXT("IoDispatcher decompression worker count.")
+);
+
+int32 GIoDispatcherCacheSizeMB = 0;
+static FAutoConsoleVariableRef CVar_IoDispatcherCacheSizeMB(
+	TEXT("s.IoDispatcherCacheSizeMB"),
+	GIoDispatcherCacheSizeMB,
+	TEXT("IoDispatcher cache memory size (in megabytes).")
+);
+
+int32 GIoDispatcherSortRequestsByOffset = 1;
+static FAutoConsoleVariableRef CVar_IoDispatcherSortRequestsByOffset(
+	TEXT("s.IoDispatcherSortRequestsByOffset"),
+	GIoDispatcherSortRequestsByOffset,
+	TEXT("If > 0, io dispatcher sorts the outstanding request queue by offset rather than sequence.")
+);
+
+int32 GIoDispatcherMaintainSortingOnPriorityChange = 1;
+static FAutoConsoleVariableRef CVar_IoDispatcherMaintainSortingOnPriorityChange(
+	TEXT("s.IoDispatcherMaintainSortingOnPriorityChange"),
+	GIoDispatcherMaintainSortingOnPriorityChange,
+	TEXT("If s.IoDispatcherSortRequestsByOffset > 0 and this > 0, io dispatcher remembers the last file handle/offset read from even when switching priority levels.")
+);
+
+int32 GIoDispatcherMaxForwardSeekKB = 0;
+static FAutoConsoleVariableRef CVar_IoDispatcherMaxForwardSeekKB(
+	TEXT("s.IoDispatcherMaxForwardSeekKB"),
+	GIoDispatcherMaxForwardSeekKB,
+	TEXT("If s.IoDispatcherSortRequestsByOffset is enabled and this is > 0, if the next sequential read is further than this offset from the last one, read the oldest request instead")
+);
+
+int32 GIoDispatcherRequestLatencyCircuitBreakerMS = 0;
+static FAutoConsoleVariableRef CVar_IoDispatcherRequestLatencyCircuitBreakerMS(
+	TEXT("s.IoDispatcherRequestLatencyCircuitBreakerMS"),
+	GIoDispatcherRequestLatencyCircuitBreakerMS,
+	TEXT("If s.IoDispatcherSortRequestsByOffset is enabled and this is >0, if the oldest request has been in the queue for this long, read it instead of the most optimal read")
+);
+
+int32 GIoDispatcherTocsEnablePerfectHashing = 1;
+static FAutoConsoleVariableRef CVar_IoDispatcherTocsEnablePerfectHashing(
+	TEXT("s.IoDispatcherTocsEnablePerfectHashing"),
+	GIoDispatcherTocsEnablePerfectHashing,
+	TEXT("Enable perfect hashmap lookups for iostore tocs")
+);
+
+uint32 FFileIoStoreReadRequest::NextSequence = 0;
+#if CHECK_IO_STORE_READ_REQUEST_LIST_MEMBERSHIP
+uint32 FFileIoStoreReadRequestList::NextListCookie = 0;
+#endif
+TAtomic<uint32> FFileIoStoreReader::GlobalPartitionIndex{ 0 };
+TAtomic<uint32> FFileIoStoreReader::GlobalContainerInstanceId{ 0 };
+
+class FMappedFileProxy final : public IMappedFileHandle
+{
+public:
+	FMappedFileProxy(IMappedFileHandle* InSharedMappedFileHandle, uint64 InSize)
+		: IMappedFileHandle(InSize)
+		, SharedMappedFileHandle(InSharedMappedFileHandle)
+	{
+		check(InSharedMappedFileHandle != nullptr);
+	}
+
+	virtual ~FMappedFileProxy() { }
+
+	virtual IMappedFileRegion* MapRegion(int64 Offset = 0, int64 BytesToMap = MAX_int64, bool bPreloadHint = false) override
+	{
+		return SharedMappedFileHandle->MapRegion(Offset, BytesToMap, bPreloadHint);
+	}
+private:
+	IMappedFileHandle* SharedMappedFileHandle;
+};
+
+void FFileIoStoreBufferAllocator::Initialize(uint64 InMemorySize, uint64 InBufferSize, uint32 InBufferAlignment)
+{
+	uint64 BufferCount = InMemorySize / InBufferSize;
+	uint64 MemorySize = BufferCount * InBufferSize;
+	BufferMemory = reinterpret_cast<uint8*>(FMemory::Malloc(MemorySize, InBufferAlignment));
+	BufferSize = InBufferSize;
+	for (uint64 BufferIndex = 0; BufferIndex < BufferCount; ++BufferIndex)
+	{
+		FFileIoStoreBuffer* Buffer = new FFileIoStoreBuffer();
+		Buffer->Memory = BufferMemory + BufferIndex * BufferSize;
+		Buffer->Next = FirstFreeBuffer;
+		FirstFreeBuffer = Buffer;
+		Stats.OnBufferReleased();
+	}
+}
+
+FFileIoStoreBuffer* FFileIoStoreBufferAllocator::AllocBuffer()
+{
+	FFileIoStoreBuffer* Buffer;
+	{
+		FScopeLock Lock(&BuffersCritical);
+		Buffer = FirstFreeBuffer;
+		if (!Buffer)
+		{
+			return nullptr;
+		}
+		FirstFreeBuffer = Buffer->Next;
+	}
+	Stats.OnBufferAllocated();
+	return Buffer;
+}
+
+void FFileIoStoreBufferAllocator::FreeBuffer(FFileIoStoreBuffer* Buffer)
+{
+	check(Buffer);
+	{
+		FScopeLock Lock(&BuffersCritical);
+		Buffer->Next = FirstFreeBuffer;
+		FirstFreeBuffer = Buffer;
+	}
+	Stats.OnBufferReleased();
+}
+
+FFileIoStoreBlockCache::FFileIoStoreBlockCache(FFileIoStoreStats& InStats)
+	: Stats(InStats)
+{
+	CacheLruHead.LruNext = &CacheLruTail;
+	CacheLruTail.LruPrev = &CacheLruHead;
+}
+
+FFileIoStoreBlockCache::~FFileIoStoreBlockCache()
+{
+	FCachedBlock* CachedBlock = CacheLruHead.LruNext;
+	while (CachedBlock != &CacheLruTail)
+	{
+		FCachedBlock* Next = CachedBlock->LruNext;
+		delete CachedBlock;
+		CachedBlock = Next;
+	}
+	FMemory::Free(CacheMemory);
+}
+
+void FFileIoStoreBlockCache::Initialize(uint64 InCacheMemorySize, uint64 InReadBufferSize)
+{
+	ReadBufferSize = InReadBufferSize;
+	uint64 CacheBlockCount = InCacheMemorySize / InReadBufferSize;
+	if (CacheBlockCount)
+	{
+		InCacheMemorySize = CacheBlockCount * InReadBufferSize;
+		CacheMemory = reinterpret_cast<uint8*>(FMemory::Malloc(InCacheMemorySize));
+		FCachedBlock* Prev = &CacheLruHead;
+		for (uint64 CacheBlockIndex = 0; CacheBlockIndex < CacheBlockCount; ++CacheBlockIndex)
+		{
+			FCachedBlock* CachedBlock = new FCachedBlock();
+			CachedBlock->Key = uint64(-1);
+			CachedBlock->Buffer = CacheMemory + CacheBlockIndex * InReadBufferSize;
+			Prev->LruNext = CachedBlock;
+			CachedBlock->LruPrev = Prev;
+			Prev = CachedBlock;
+		}
+		Prev->LruNext = &CacheLruTail;
+		CacheLruTail.LruPrev = Prev;
+	}
+}
+
+bool FFileIoStoreBlockCache::Read(FFileIoStoreReadRequest* Block)
+{
+	if (!CacheMemory)
+	{
+		return false;
+	}
+	check(Block->Buffer);
+	FCachedBlock* CachedBlock = CachedBlocks.FindRef(Block->Key.Hash);
+	if (!CachedBlock)
+	{
+		Stats.OnBlockCacheMiss(ReadBufferSize);
+		return false;
+	}
+	
+	CachedBlock->LruPrev->LruNext = CachedBlock->LruNext;
+	CachedBlock->LruNext->LruPrev = CachedBlock->LruPrev;
+
+	CachedBlock->LruPrev = &CacheLruHead;
+	CachedBlock->LruNext = CacheLruHead.LruNext;
+
+	CachedBlock->LruPrev->LruNext = CachedBlock;
+	CachedBlock->LruNext->LruPrev = CachedBlock;
+
+	check(CachedBlock->Buffer);
+	Stats.OnBlockCacheHit(ReadBufferSize);
+	FMemory::Memcpy(Block->Buffer->Memory, CachedBlock->Buffer, ReadBufferSize);
+
+	return true;
+}
+
+void FFileIoStoreBlockCache::Store(const FFileIoStoreReadRequest* Block)
+{
+	bool bIsCacheableBlock = CacheMemory != nullptr && Block->BytesUsed < Block->Size;
+	if (!bIsCacheableBlock)
+	{
+		return;
+	}
+	check(Block->Buffer);
+	check(Block->Buffer->Memory);
+	FCachedBlock* BlockToReplace = CacheLruTail.LruPrev;
+	if (BlockToReplace == &CacheLruHead)
+	{
+		return;
+	}
+	check(BlockToReplace);
+	CachedBlocks.Remove(BlockToReplace->Key);
+	BlockToReplace->Key = Block->Key.Hash;
+
+	BlockToReplace->LruPrev->LruNext = BlockToReplace->LruNext;
+	BlockToReplace->LruNext->LruPrev = BlockToReplace->LruPrev;
+
+	BlockToReplace->LruPrev = &CacheLruHead;
+	BlockToReplace->LruNext = CacheLruHead.LruNext;
+
+	BlockToReplace->LruPrev->LruNext = BlockToReplace;
+	BlockToReplace->LruNext->LruPrev = BlockToReplace;
+	
+	check(BlockToReplace->Buffer);
+	FMemory::Memcpy(BlockToReplace->Buffer, Block->Buffer->Memory, ReadBufferSize);
+	Stats.OnBlockCacheStore(ReadBufferSize);
+	CachedBlocks.Add(BlockToReplace->Key, BlockToReplace);
+}
+
+#define UE_FILEIOSTORE_DETAILED_QUEUE_COUNTERS_ENABLED 0
+#if UE_FILEIOSTORE_DETAILED_QUEUE_COUNTERS_ENABLED
+TRACE_DECLARE_INT_COUNTER(IoDispatcherLatencyCircuitBreaks, TEXT("IoDispatcher/LatencyCircuitBreaks"));
+TRACE_DECLARE_INT_COUNTER(IoDispatcherSeekDistanceCircuitBreaks, TEXT("IoDispatcher/SeekDistanceCircuitBreaks"));
+TRACE_DECLARE_INT_COUNTER(IoDispatcherNumPriorityQueues, TEXT("IoDispatcher/NumPriorityQueues"));
+#endif
+
+bool FFileIoStoreOffsetSortedRequestQueue::RequestSortPredicate(const FFileIoStoreReadRequestSortKey& A, const FFileIoStoreReadRequestSortKey& B)
+{
+	if (A.Handle == B.Handle)
+	{
+		return A.Offset < B.Offset;
+	}
+	return A.Handle < B.Handle;
+}
+
+FFileIoStoreOffsetSortedRequestQueue::FFileIoStoreOffsetSortedRequestQueue(int32 InPriority)
+	: Priority(InPriority)
+{
+}
+
+TArray<FFileIoStoreReadRequest*> FFileIoStoreOffsetSortedRequestQueue::StealRequests()
+{
+	RequestsBySequence.Clear();
+	PeekRequestIndex = INDEX_NONE;
+	return MoveTemp(Requests); 
+}
+
+// This could be potentially optimized if the higher level keeps track of which requests it changes the priority of, or even just the old priorty levels
+TArray<FFileIoStoreReadRequest*> FFileIoStoreOffsetSortedRequestQueue::RemoveMisprioritizedRequests()
+{
+	PeekRequestIndex = INDEX_NONE;
+	TArray<FFileIoStoreReadRequest*> RequestsToReturn;
+	for (int32 i = Requests.Num()-1; i >= 0; --i)
+	{
+		if (Requests[i]->Priority != Priority)
+		{
+			RequestsToReturn.Add(Requests[i]);
+			RequestsBySequence.Remove(Requests[i]);
+			Requests.RemoveAt(i, 1, false);
+		}
+	}
+
+	return RequestsToReturn;
+}
+
+FFileIoStoreReadRequest* FFileIoStoreOffsetSortedRequestQueue::GetNextInternal(FFileIoStoreReadRequestSortKey LastSortKey, bool bPop)
+{
+	if (Requests.Num() == 0)
+	{
+		return nullptr;
+	}
+
+	int32 RequestIndex = INDEX_NONE;
+	if (PeekRequestIndex != INDEX_NONE)
+	{
+		RequestIndex = PeekRequestIndex;
+	}
+	else 
+	{
+		bool bHeadRequestTooOld = false;
+		if (GIoDispatcherRequestLatencyCircuitBreakerMS > 0)
+		{
+			// If our oldest request has been unserviced for too long, grab that instead of the next sequential read
+			uint64 ThresholdCycles = uint64((GIoDispatcherRequestLatencyCircuitBreakerMS * 1000.0) / FPlatformTime::GetSecondsPerCycle64());
+			bHeadRequestTooOld = (FPlatformTime::Cycles64() - RequestsBySequence.PeekHead()->CreationTime) >= ThresholdCycles;
+
+#if UE_FILEIOSTORE_DETAILED_QUEUE_COUNTERS_ENABLED
+			if (bPop)
+			{
+				TRACE_COUNTER_INCREMENT(IoDispatcherLatencyCircuitBreaks);
+			}
+#endif
+		}
+
+		const bool bChooseByOffset = 
+				LastSortKey.Handle != 0 
+			&&	!bHeadRequestTooOld 
+			&&  (GIoDispatcherMaintainSortingOnPriorityChange || LastSortKey.Priority == Priority);
+		if (bChooseByOffset)
+		{
+			// Pick the request with the closest offset to the last thing that we read
+			RequestIndex = Algo::LowerBoundBy(Requests, LastSortKey, RequestSortProjection, RequestSortPredicate);
+			if (Requests.IsValidIndex(RequestIndex)) // If all requests are before LastOffset we get back out-of-bounds
+			{
+				if (Requests[RequestIndex]->FileHandle != LastSortKey.Handle)
+				{
+					// Changing file handle so switch back to the oldest outstanding request 
+					RequestIndex = INDEX_NONE;
+				}
+				else if (GIoDispatcherMaxForwardSeekKB > 0 && (LastSortKey.Offset - Requests[RequestIndex]->Offset) > uint64(GIoDispatcherMaxForwardSeekKB) * 1024)
+				{
+					// Large forward seek so switch back to the oldest outstanding request 
+					RequestIndex = INDEX_NONE;
+
+#if UE_FILEIOSTORE_DETAILED_QUEUE_COUNTERS_ENABLED
+					if (bPop)
+					{
+						TRACE_COUNTER_INCREMENT(IoDispatcherSeekDistanceCircuitBreaks);
+					}
+#endif
+				}
+			}
+		}
+
+		if (!Requests.IsValidIndex(RequestIndex))
+		{
+			RequestIndex = Requests.Find(RequestsBySequence.PeekHead());
+			check(Requests[RequestIndex] == RequestsBySequence.PeekHead());
+		}
+	}
+
+	check(Requests.IsValidIndex(RequestIndex));
+
+	FFileIoStoreReadRequest* Request = Requests[RequestIndex];
+	if (bPop)
+	{
+		Requests.RemoveAt(RequestIndex);
+		RequestsBySequence.Remove(Request);
+		PeekRequestIndex = INDEX_NONE;
+	}
+	else
+	{
+		PeekRequestIndex = RequestIndex;
+	}
+	return Request;
+}
+
+FFileIoStoreReadRequest* FFileIoStoreOffsetSortedRequestQueue::Peek(FFileIoStoreReadRequestSortKey LastSortKey)
+{
+	return GetNextInternal(LastSortKey, false);
+}
+
+FFileIoStoreReadRequest* FFileIoStoreOffsetSortedRequestQueue::Pop(FFileIoStoreReadRequestSortKey LastSortKey)
+{
+	return GetNextInternal(LastSortKey, true);
+}
+
+void FFileIoStoreOffsetSortedRequestQueue::Push(FFileIoStoreReadRequest* Request)
+{
+	// Insert sorted by file handle & offset
+	int32 InsertIndex = Algo::UpperBoundBy(Requests, RequestSortProjection(Request), RequestSortProjection, RequestSortPredicate);
+	Requests.Insert(Request, InsertIndex);
+	
+	// Insert sorted by age
+	RequestsBySequence.Add(Request);
+
+	PeekRequestIndex = INDEX_NONE;
+}
+
+void FFileIoStoreOffsetSortedRequestQueue::CancelRequestsWithFileHandle(uint64 FileHandle)
+{
+	for (FFileIoStoreReadRequest* Request : Requests)
+	{
+		if (Request->FileHandle == FileHandle)
+		{
+			Request->bCancelled = true;
+		}
+	}
+}
+
+void FFileIoStoreRequestQueue::UpdateSortRequestsByOffset()
+{
+	// Must hold CriticalSection here
+	if (bSortRequestsByOffset == bool(GIoDispatcherSortRequestsByOffset))
+	{
+		return;
+	}
+
+	bSortRequestsByOffset = bool(GIoDispatcherSortRequestsByOffset);
+	if (bSortRequestsByOffset)
+	{
+		// Split things into separate heaps
+		for (FFileIoStoreReadRequest* Request : Heap)
+		{
+			Push(*Request);
+		}
+		Heap.Empty();
+	}
+	else
+	{
+		// Put things back into the main heap
+		TArray< FFileIoStoreReadRequest*> AllRequests;
+		for (FFileIoStoreOffsetSortedRequestQueue& SubQueue : SortedPriorityQueues)
+		{
+			AllRequests.Append(SubQueue.StealRequests());
+		}
+		Algo::SortBy(AllRequests, [](FFileIoStoreReadRequest* Request) { return Request->Sequence; });
+		for (FFileIoStoreReadRequest* Request : AllRequests)
+		{
+			Push(*Request);
+		}
+		check(Algo::AllOf(SortedPriorityQueues, [](const FFileIoStoreOffsetSortedRequestQueue& Q) { return Q.IsEmpty(); }));
+		SortedPriorityQueues.Empty();
+	}
+}
+
+FFileIoStoreReadRequest* FFileIoStoreRequestQueue::Peek()
+{
+	TRACE_CPUPROFILER_EVENT_SCOPE(RequestQueuePeek);
+	FScopeLock _(&CriticalSection);
+	UpdateSortRequestsByOffset();
+	if (bSortRequestsByOffset)
+	{
+		if( SortedPriorityQueues.Num() == 0)
+		{
+			return nullptr;
+		}
+
+		FFileIoStoreOffsetSortedRequestQueue& SubQueue = SortedPriorityQueues.Last();
+		check(!SubQueue.IsEmpty());
+		FFileIoStoreReadRequest* Request = SubQueue.Peek(LastSortKey);
+		check(Request);
+		// Do not update LastUsedPriority here until we actually pop
+		return Request;
+	}
+	else
+	{
+		if (Heap.Num() == 0)
+		{
+			return nullptr;
+		}
+		return Heap.HeapTop();
+	}
+}
+
+FFileIoStoreReadRequest* FFileIoStoreRequestQueue::Pop()
+{
+	TRACE_CPUPROFILER_EVENT_SCOPE(RequestQueuePop);
+	FScopeLock _(&CriticalSection);
+	UpdateSortRequestsByOffset();
+	FFileIoStoreReadRequest* Result = nullptr;
+	if (bSortRequestsByOffset)
+	{
+		if (SortedPriorityQueues.Num() == 0)
+		{
+			return nullptr;
+		}
+
+		FFileIoStoreOffsetSortedRequestQueue& SubQueue = SortedPriorityQueues.Last();
+		check(!SubQueue.IsEmpty());
+		Result = SubQueue.Pop(LastSortKey);
+		check(Result);
+		LastSortKey = Result;
+		if (SubQueue.IsEmpty())
+		{
+			SortedPriorityQueues.Pop();
+			// SubQueue is invalid here
+#if UE_FILEIOSTORE_DETAILED_QUEUE_COUNTERS_ENABLED
+			TRACE_COUNTER_DECREMENT(IoDispatcherNumPriorityQueues);
+#endif
+		}
+	}
+	else
+	{
+		if (Heap.Num() == 0)
+		{
+			return nullptr;
+		}
+		Heap.HeapPop(Result, QueueSortFunc, false);
+	}
+	
+	check(Result->QueueStatus == FFileIoStoreReadRequest::QueueStatus_InQueue);
+	Result->QueueStatus = FFileIoStoreReadRequest::QueueStatus_Started;
+	return Result;
+}
+
+void FFileIoStoreRequestQueue::PushToPriorityQueues(FFileIoStoreReadRequest* Request)
+{
+	int32 QueueIndex = Algo::LowerBoundBy(SortedPriorityQueues, Request->Priority, QueuePriorityProjection, TLess<int32>());
+	if (!SortedPriorityQueues.IsValidIndex(QueueIndex) || SortedPriorityQueues[QueueIndex].GetPriority() != Request->Priority)
+	{
+		SortedPriorityQueues.Insert(FFileIoStoreOffsetSortedRequestQueue(Request->Priority), QueueIndex);
+#if UE_FILEIOSTORE_DETAILED_QUEUE_COUNTERS_ENABLED
+		TRACE_COUNTER_INCREMENT(IoDispatcherNumPriorityQueues);
+#endif
+	}
+	check(Algo::IsSortedBy(SortedPriorityQueues, QueuePriorityProjection, TLess<int32>()));
+	FFileIoStoreOffsetSortedRequestQueue& Queue = SortedPriorityQueues[QueueIndex];
+	check(Queue.GetPriority() == Request->Priority);
+	Queue.Push(Request);
+}
+
+void FFileIoStoreRequestQueue::Push(FFileIoStoreReadRequest& Request)
+{
+	TRACE_CPUPROFILER_EVENT_SCOPE(RequestQueuePush);
+	FScopeLock _(&CriticalSection);
+	UpdateSortRequestsByOffset();
+	
+	check(Request.QueueStatus != FFileIoStoreReadRequest::QueueStatus_InQueue);
+	Request.QueueStatus = FFileIoStoreReadRequest::QueueStatus_InQueue;
+
+	if (bSortRequestsByOffset)
+	{
+		PushToPriorityQueues(&Request);
+	}
+	else
+	{
+		Heap.HeapPush(&Request, QueueSortFunc);
+	}
+}
+
+void FFileIoStoreRequestQueue::Push(FFileIoStoreReadRequestList& Requests)
+{
+	TRACE_CPUPROFILER_EVENT_SCOPE(RequestQueuePush);
+	FScopeLock _(&CriticalSection);
+	UpdateSortRequestsByOffset();
+
+	for (auto It = Requests.Steal(); It; ++It)
+	{
+		check(It->QueueStatus != FFileIoStoreReadRequest::QueueStatus_InQueue);
+		It->QueueStatus = FFileIoStoreReadRequest::QueueStatus_InQueue;
+
+		if (bSortRequestsByOffset)
+		{
+			PushToPriorityQueues(*It);
+		}
+		else
+		{
+			Heap.HeapPush(*It, QueueSortFunc);
+		}
+	}
+}
+
+void FFileIoStoreRequestQueue::UpdateOrder()
+{
+	TRACE_CPUPROFILER_EVENT_SCOPE(RequestQueueUpdateOrder);
+	FScopeLock _(&CriticalSection);
+	UpdateSortRequestsByOffset();
+	if (bSortRequestsByOffset)
+	{
+		TArray<FFileIoStoreReadRequest*> Requests;
+		for (FFileIoStoreOffsetSortedRequestQueue& SubQueue : SortedPriorityQueues)
+		{
+			TArray<FFileIoStoreReadRequest*> RequestsRemoved = SubQueue.RemoveMisprioritizedRequests();
+			Requests.Append(RequestsRemoved);
+		}
+
+		// Pop/Peek rely on empty queues being culled
+		SortedPriorityQueues.RemoveAll([](FFileIoStoreOffsetSortedRequestQueue& SubQueue) { return SubQueue.IsEmpty(); });
+
+		Algo::SortBy(Requests, [](FFileIoStoreReadRequest* Request) { return Request->Sequence;  });
+		for (FFileIoStoreReadRequest* Request : Requests)
+		{
+			PushToPriorityQueues(Request);
+		}
+	}
+	else
+	{
+		Heap.Heapify(QueueSortFunc);
+	}
+}
+
+void FFileIoStoreRequestQueue::Lock()
+{
+	CriticalSection.Lock();
+}
+
+void FFileIoStoreRequestQueue::Unlock()
+{
+	CriticalSection.Unlock();
+}
+
+void FFileIoStoreRequestQueue::CancelRequestsWithFileHandle(const uint64 FileHandle)
+{
+	FScopeLock _(&CriticalSection);
+	
+	if (bSortRequestsByOffset)
+	{
+		for (FFileIoStoreOffsetSortedRequestQueue& SubQueue : SortedPriorityQueues)
+		{
+			SubQueue.CancelRequestsWithFileHandle(FileHandle);
+		}
+	}
+	else
+	{
+		for (FFileIoStoreReadRequest* Request : Heap)
+		{
+			if (Request->FileHandle == FileHandle)
+			{
+				Request->bCancelled = true;
+			}
+		}
+	}
+}
+
+FFileIoStoreReader::FFileIoStoreReader(IPlatformFileIoStore& InPlatformImpl, FFileIoStoreStats& InStats)
+	: PlatformImpl(InPlatformImpl)
+	, Stats(InStats)
+{
+}
+
+FFileIoStoreReader::~FFileIoStoreReader()
+{
+	Close();
+}
+
+uint64 FFileIoStoreReader::GetTocAllocatedSize() const
+{
+	return TocImperfectHashMapFallback.GetAllocatedSize() +
+		PerfectHashMap.TocOffsetAndLengths.GetAllocatedSize() +
+		PerfectHashMap.TocChunkIds.GetAllocatedSize() +
+		PerfectHashMap.TocChunkHashSeeds.GetAllocatedSize() +
+		ContainerFile.CompressionBlocks.GetAllocatedSize() +
+		ContainerFile.BlockSignatureHashes.GetAllocatedSize();
+}
+
+FIoStatus FFileIoStoreReader::Initialize(const TCHAR* InTocFilePath, int32 InOrder)
+{
+	FStringView ContainerPathView(InTocFilePath);
+	if (!ContainerPathView.EndsWith(TEXT(".utoc")))
+	{
+		return FIoStatusBuilder(EIoErrorCode::FileOpenFailed) << TEXT("Expected .utoc extension on container path '") << InTocFilePath << TEXT("'");
+	}
+	FStringView BasePathView = ContainerPathView.LeftChop(5);
+
+	IPlatformFile& Ipf = FPlatformFileManager::Get().GetPlatformFile();
+
+	UE_LOG(LogIoDispatcher, Display, TEXT("Reading toc: %s"), InTocFilePath);
+
+	TUniquePtr<FIoStoreTocResource> TocResourcePtr = MakeUnique<FIoStoreTocResource>();
+	FIoStoreTocResource& TocResource = *TocResourcePtr;
+	FIoStatus Status = FIoStoreTocResource::Read(InTocFilePath, EIoStoreTocReadOptions::Default, TocResource);
+	if (!Status.IsOk())
+	{
+		return Status;
+	}
+
+	ContainerFile.PartitionSize = TocResource.Header.PartitionSize;
+	ContainerFile.Partitions.SetNum(TocResource.Header.PartitionCount);
+	for (uint32 PartitionIndex = 0; PartitionIndex < TocResource.Header.PartitionCount; ++PartitionIndex)
+	{
+		FFileIoStoreContainerFilePartition& Partition = ContainerFile.Partitions[PartitionIndex];
+		TStringBuilder<256> ContainerFilePath;
+		ContainerFilePath.Append(BasePathView);
+		if (PartitionIndex > 0)
+		{
+			ContainerFilePath.Appendf(TEXT("_s%d"), PartitionIndex);
+		}
+		ContainerFilePath.Append(TEXT(".ucas"));
+		Partition.FilePath = ContainerFilePath;
+		if (!PlatformImpl.OpenContainer(*ContainerFilePath, Partition.FileHandle, Partition.FileSize))
+		{
+			return FIoStatusBuilder(EIoErrorCode::FileOpenFailed) << TEXT("Failed to open IoStore container file '") << *ContainerFilePath << TEXT("'");
+		}
+		Partition.ContainerFileIndex = GlobalPartitionIndex++;
+	}
+
+	if (GIoDispatcherTocsEnablePerfectHashing && !TocResource.ChunkPerfectHashSeeds.IsEmpty())
+	{
+		for (int32 ChunkIndexWithoutPerfectHash : TocResource.ChunkIndicesWithoutPerfectHash)
+		{
+			TocImperfectHashMapFallback.Add(TocResource.ChunkIds[ChunkIndexWithoutPerfectHash], TocResource.ChunkOffsetLengths[ChunkIndexWithoutPerfectHash]);
+		}
+
+		PerfectHashMap.TocChunkHashSeeds = MoveTemp(TocResource.ChunkPerfectHashSeeds);
+		PerfectHashMap.TocOffsetAndLengths = MoveTemp(TocResource.ChunkOffsetLengths);
+		PerfectHashMap.TocChunkIds = MoveTemp(TocResource.ChunkIds);
+		bHasPerfectHashMap = true;
+	}
+	else
+	{
+		UE_LOG(LogIoDispatcher, Warning, TEXT("Falling back to imperfect hashmap for container '%s'"), InTocFilePath);
+		for (uint32 ChunkIndex = 0; ChunkIndex < TocResource.Header.TocEntryCount; ++ChunkIndex)
+		{
+			TocImperfectHashMapFallback.Add(TocResource.ChunkIds[ChunkIndex], TocResource.ChunkOffsetLengths[ChunkIndex]);
+		}
+		bHasPerfectHashMap = false;
+	}
+	
+	ContainerFile.CompressionMethods	= MoveTemp(TocResource.CompressionMethods);
+	ContainerFile.CompressionBlockSize	= TocResource.Header.CompressionBlockSize;
+	ContainerFile.CompressionBlocks		= MoveTemp(TocResource.CompressionBlocks);
+	ContainerFile.ContainerFlags		= TocResource.Header.ContainerFlags;
+	ContainerFile.EncryptionKeyGuid		= TocResource.Header.EncryptionKeyGuid;
+	ContainerFile.BlockSignatureHashes	= MoveTemp(TocResource.ChunkBlockSignatures);
+	ContainerFile.ContainerInstanceId	= ++GlobalContainerInstanceId;
+
+	Stats.OnTocMounted(GetTocAllocatedSize());
+
+	UE_LOG(LogIoDispatcher, Display, TEXT("Toc signature hash: %s"), *TocResource.SignatureHash.ToString());
+	
+	ContainerId = TocResource.Header.ContainerId;
+	Order = InOrder;
+	return FIoStatus::Ok;
+}
+
+FIoStatus FFileIoStoreReader::Close()
+{
+	if (bClosed)
+	{
+		return FIoStatus::Ok;
+	}
+
+	for (FFileIoStoreContainerFilePartition& Partition : ContainerFile.Partitions)
+	{
+		PlatformImpl.CloseContainer(Partition.FileHandle);
+	}
+
+	Stats.OnTocUnmounted(GetTocAllocatedSize());
+
+	PerfectHashMap.TocChunkHashSeeds.Empty();
+	PerfectHashMap.TocChunkIds.Empty();
+	PerfectHashMap.TocOffsetAndLengths.Empty();
+	TocImperfectHashMapFallback.Empty();
+	ContainerFile = FFileIoStoreContainerFile();
+	ContainerId = FIoContainerId();
+	Order = INDEX_NONE;
+	bClosed = true;
+
+	return FIoStatus::Ok;
+}
+
+const FIoOffsetAndLength* FFileIoStoreReader::FindChunkInternal(const FIoChunkId& ChunkId) const
+{
+	if (bHasPerfectHashMap)
+	{
+		// See FIoStoreWriterImpl::GeneratePerfectHashes
+		const uint32 ChunkCount = PerfectHashMap.TocChunkIds.Num();
+		if (!ChunkCount)
+		{
+			return nullptr;
+		}
+		const uint32 SeedCount = PerfectHashMap.TocChunkHashSeeds.Num();
+		uint32 SeedIndex = FIoStoreTocResource::HashChunkIdWithSeed(0, ChunkId) % SeedCount;
+		const int32 Seed = PerfectHashMap.TocChunkHashSeeds[SeedIndex];
+		if (Seed == 0)
+		{
+			return nullptr;
+		}
+		uint32 Slot;
+		if (Seed < 0)
+		{
+			uint32 SeedAsIndex = static_cast<uint32>(-Seed - 1);
+			if (SeedAsIndex < ChunkCount)
+			{
+				Slot = static_cast<uint32>(SeedAsIndex);
+			}
+			else
+			{
+				// Entry without perfect hash
+				return TocImperfectHashMapFallback.Find(ChunkId);
+			}
+		}
+		else
+		{
+			Slot = FIoStoreTocResource::HashChunkIdWithSeed(static_cast<uint32>(Seed), ChunkId) % ChunkCount;
+		}
+		if (PerfectHashMap.TocChunkIds[Slot] == ChunkId)
+		{
+			return &PerfectHashMap.TocOffsetAndLengths[Slot];
+		}
+		return nullptr;
+	}
+	else
+	{
+		return TocImperfectHashMapFallback.Find(ChunkId);
+	}
+}
+
+bool FFileIoStoreReader::DoesChunkExist(const FIoChunkId& ChunkId) const
+{
+	check(!bClosed);
+	return FindChunkInternal(ChunkId) != nullptr;
+}
+
+TIoStatusOr<uint64> FFileIoStoreReader::GetSizeForChunk(const FIoChunkId& ChunkId) const
+{
+	check(!bClosed);
+	const FIoOffsetAndLength* OffsetAndLength = FindChunkInternal(ChunkId);
+	if (OffsetAndLength)
+	{
+		return OffsetAndLength->GetLength();
+	}
+	else
+	{
+		return FIoStatus(EIoErrorCode::NotFound);
+	}
+}
+
+const FIoOffsetAndLength* FFileIoStoreReader::Resolve(const FIoChunkId& ChunkId) const
+{
+	check(!bClosed);
+	return FindChunkInternal(ChunkId);
+}
+
+IMappedFileHandle* FFileIoStoreReader::GetMappedContainerFileHandle(uint64 TocOffset)
+{
+	check(!bClosed);
+	int32 PartitionIndex = int32(TocOffset / ContainerFile.PartitionSize);
+	FFileIoStoreContainerFilePartition& Partition = ContainerFile.Partitions[PartitionIndex];
+	if (!Partition.MappedFileHandle)
+	{
+		IPlatformFile& Ipf = FPlatformFileManager::Get().GetPlatformFile();
+		Partition.MappedFileHandle.Reset(Ipf.OpenMapped(*Partition.FilePath));
+	}
+
+	check(Partition.FileSize > 0);
+	return new FMappedFileProxy(Partition.MappedFileHandle.Get(), Partition.FileSize);
+}
+
+TIoStatusOr<FIoContainerHeader> FFileIoStoreReader::ReadContainerHeader() const
+{
+	LLM_SCOPE(ELLMTag::AsyncLoading);
+	TRACE_CPUPROFILER_EVENT_SCOPE(ReadContainerHeader);
+	FIoChunkId HeaderChunkId = CreateIoChunkId(ContainerId.Value(), 0, EIoChunkType::ContainerHeader);
+	const FIoOffsetAndLength* OffsetAndLength = FindChunkInternal(HeaderChunkId);
+	if (!OffsetAndLength)
+	{
+		return FIoStatus(EIoErrorCode::NotFound);
+	}
+	
+	const uint64 CompressionBlockSize = ContainerFile.CompressionBlockSize;
+	const uint64 Offset = OffsetAndLength->GetOffset();
+	const uint64 Size = OffsetAndLength->GetLength();
+	const uint64 RequestEndOffset = Offset + Size;
+	const int32 RequestBeginBlockIndex = int32(Offset / CompressionBlockSize);
+	const int32 RequestEndBlockIndex = int32((RequestEndOffset - 1) / CompressionBlockSize);
+	
+	// Assumes that the container header is uncompressed and placed in its own blocks in the same partition without padding
+	const FIoStoreTocCompressedBlockEntry* CompressionBlockEntry = &ContainerFile.CompressionBlocks[RequestBeginBlockIndex];
+	const int32 PartitionIndex = int32(CompressionBlockEntry->GetOffset() / ContainerFile.PartitionSize);
+	const FFileIoStoreContainerFilePartition& Partition = ContainerFile.Partitions[PartitionIndex];
+	const uint64 RawOffset = CompressionBlockEntry->GetOffset() % ContainerFile.PartitionSize;
+
+	FIoBuffer IoBuffer(Align(Size, FAES::AESBlockSize));
+	IPlatformFile& Ipf = FPlatformFileManager::Get().GetPlatformFile();
+	TUniquePtr<IFileHandle> ContainerFileHandle(Ipf.OpenRead(*Partition.FilePath));
+	if (!ContainerFileHandle)
+	{
+		return FIoStatus(EIoErrorCode::FileOpenFailed);
+	}
+	{
+		TRACE_CPUPROFILER_EVENT_SCOPE(ReadFromContainerFile);
+		if (!ContainerFileHandle->Seek(RawOffset))
+		{
+			return FIoStatus(EIoErrorCode::ReadError);
+		}
+		if (!ContainerFileHandle->Read(IoBuffer.Data(), IoBuffer.DataSize()))
+		{
+			return FIoStatus(EIoErrorCode::ReadError);
+		}
+	}
+
+	const bool bSigned = EnumHasAnyFlags(ContainerFile.ContainerFlags, EIoContainerFlags::Signed);
+	const bool bEncrypted = ContainerFile.EncryptionKey.IsValid();
+	if (bSigned || bEncrypted)
+	{
+		uint8* BlockData = IoBuffer.Data();
+		for (int32 CompressedBlockIndex = RequestBeginBlockIndex; CompressedBlockIndex <= RequestEndBlockIndex; ++CompressedBlockIndex)
+		{
+			CompressionBlockEntry = &ContainerFile.CompressionBlocks[CompressedBlockIndex];
+			check(ContainerFile.CompressionMethods[CompressionBlockEntry->GetCompressionMethodIndex()].IsNone());
+			const uint64 BlockSize = Align(CompressionBlockEntry->GetCompressedSize(), FAES::AESBlockSize);
+			if (bSigned)
+			{
+				const FSHAHash& SignatureHash = ContainerFile.BlockSignatureHashes[CompressedBlockIndex];
+				FSHAHash BlockHash;
+				FSHA1::HashBuffer(BlockData, BlockSize, BlockHash.Hash);
+				if (SignatureHash != BlockHash)
+				{
+					return FIoStatus(EIoErrorCode::SignatureError);
+				}
+			}
+			if (bEncrypted)
+			{
+				FAES::DecryptData(BlockData, uint32(BlockSize), ContainerFile.EncryptionKey);
+			}
+			BlockData += BlockSize;
+		}
+	}
+	FMemoryReaderView Ar(MakeArrayView(IoBuffer.Data(), static_cast<int32>(IoBuffer.DataSize())));
+	FIoContainerHeader ContainerHeader;
+	Ar << ContainerHeader;
+	if (Ar.IsError())
+	{
+		UE_LOG(LogIoDispatcher, Warning, TEXT("Invalid container header in file '%s'"), *ContainerFile.FilePath);
+		ContainerHeader = FIoContainerHeader();
+	}
+	return ContainerHeader;
+}
+
+void FFileIoStoreReader::ReopenAllFileHandles()
+{
+	for (FFileIoStoreContainerFilePartition& Partition : ContainerFile.Partitions)
+	{
+		PlatformImpl.CloseContainer(Partition.FileHandle);
+		PlatformImpl.OpenContainer(*Partition.FilePath, Partition.FileHandle, Partition.FileSize);
+	}
+}
+
+FFileIoStoreResolvedRequest::FFileIoStoreResolvedRequest(
+	FIoRequestImpl& InDispatcherRequest,
+	const FFileIoStoreContainerFile& InContainerFile,
+	uint64 InResolvedOffset,
+	uint64 InResolvedSize)
+	: DispatcherRequest(&InDispatcherRequest)
+	, ContainerFile(InContainerFile)
+	, ResolvedOffset(InResolvedOffset)
+	, ResolvedSize(InResolvedSize)
+{
+
+}
+
+void FFileIoStoreResolvedRequest::AddReadRequestLink(FFileIoStoreReadRequestLink* ReadRequestLink)
+{
+	check(!ReadRequestLink->Next);
+	if (ReadRequestsTail)
+	{
+		ReadRequestsTail->Next = ReadRequestLink;
+	}
+	else
+	{
+		ReadRequestsHead = ReadRequestLink;
+	}
+	ReadRequestsTail = ReadRequestLink;
+}
+
+FFileIoStoreRequestTracker::FFileIoStoreRequestTracker(FFileIoStoreRequestAllocator& InRequestAllocator, FFileIoStoreRequestQueue& InRequestQueue)
+	: RequestAllocator(InRequestAllocator)
+	, RequestQueue(InRequestQueue)
+{
+
+}
+
+FFileIoStoreRequestTracker::~FFileIoStoreRequestTracker()
+{
+
+}
+
+FFileIoStoreCompressedBlock* FFileIoStoreRequestTracker::FindOrAddCompressedBlock(FFileIoStoreBlockKey Key, bool& bOutWasAdded)
+{
+	bOutWasAdded = false;
+	FFileIoStoreCompressedBlock*& Result = CompressedBlocksMap.FindOrAdd(Key);
+	if (!Result)
+	{
+		Result = RequestAllocator.AllocCompressedBlock();
+		Result->Key = Key;
+		bOutWasAdded = true;
+	}
+	return Result;
+}
+
+FFileIoStoreReadRequest* FFileIoStoreRequestTracker::FindOrAddRawBlock(FFileIoStoreBlockKey Key, bool& bOutWasAdded)
+{
+	bOutWasAdded = false;
+	FFileIoStoreReadRequest*& Result = RawBlocksMap.FindOrAdd(Key);
+	if (!Result)
+	{
+		Result = RequestAllocator.AllocReadRequest();
+		Result->Key = Key;
+		bOutWasAdded = true;
+	}
+	return Result;
+}
+
+void FFileIoStoreRequestTracker::RemoveRawBlock(const FFileIoStoreReadRequest* RawBlock)
+{
+	if (!RawBlock->bCancelled)
+	{
+		RawBlocksMap.Remove(RawBlock->Key);
+	}
+}
+
+void FFileIoStoreRequestTracker::AddReadRequestsToResolvedRequest(FFileIoStoreCompressedBlock* CompressedBlock, FFileIoStoreResolvedRequest& ResolvedRequest)
+{
+	//TRACE_CPUPROFILER_EVENT_SCOPE(AddReadRequestsToResolvedRequest);
+	bool bUpdateQueueOrder = false;
+	++ResolvedRequest.UnfinishedReadsCount;
+	for (FFileIoStoreReadRequest* ReadRequest : CompressedBlock->RawBlocks)
+	{
+		FFileIoStoreReadRequestLink* Link = RequestAllocator.AllocRequestLink(ReadRequest);
+		++ReadRequest->RefCount;
+		ResolvedRequest.AddReadRequestLink(Link);
+		if (ResolvedRequest.GetPriority() > ReadRequest->Priority)
+		{
+			ReadRequest->Priority = ResolvedRequest.GetPriority();
+			bUpdateQueueOrder = true;
+		}
+	}
+	if (bUpdateQueueOrder)
+	{
+		TRACE_CPUPROFILER_EVENT_SCOPE(RequestTrackerAddIoRequestUpdateOrder);
+		RequestQueue.UpdateOrder();
+	}
+}
+
+void FFileIoStoreRequestTracker::AddReadRequestsToResolvedRequest(const FFileIoStoreReadRequestList& Requests, FFileIoStoreResolvedRequest& ResolvedRequest)
+{
+	//TRACE_CPUPROFILER_EVENT_SCOPE(RequestTrackerAddIoRequest);
+	for (FFileIoStoreReadRequest* Request : Requests)
+	{
+		++ResolvedRequest.UnfinishedReadsCount;
+		FFileIoStoreReadRequestLink* Link = RequestAllocator.AllocRequestLink(Request);
+		++Request->RefCount;
+		ResolvedRequest.AddReadRequestLink(Link);
+		check(ResolvedRequest.GetPriority() == Request->Priority);
+	}
+}
+
+void FFileIoStoreRequestTracker::RemoveCompressedBlock(const FFileIoStoreCompressedBlock* CompressedBlock)
+{
+	if (!CompressedBlock->bCancelled)
+	{
+		CompressedBlocksMap.Remove(CompressedBlock->Key);
+	}
+}
+
+bool FFileIoStoreRequestTracker::CancelIoRequest(FFileIoStoreResolvedRequest& ResolvedRequest)
+{
+	TRACE_CPUPROFILER_EVENT_SCOPE(RequestTrackerCancelIoRequest);
+	check(!ResolvedRequest.bCancelled);
+	bool bShouldComplete = true;
+	RequestQueue.Lock();
+	FFileIoStoreReadRequestLink* Link = ResolvedRequest.ReadRequestsHead;
+	while (Link)
+	{
+		FFileIoStoreReadRequest& ReadRequest = Link->ReadRequest;
+		Link = Link->Next;
+
+		if (ReadRequest.bCancelled)
+		{
+			continue;
+		}
+
+		if (ReadRequest.QueueStatus == FFileIoStoreReadRequest::QueueStatus_Started)
+		{
+			bShouldComplete = false;
+			continue;
+		}
+
+		bool bCancelReadRequest = true;
+		for (FFileIoStoreCompressedBlock* CompressedBlock : ReadRequest.CompressedBlocks)
+		{
+			if (CompressedBlock->bCancelled)
+			{
+				continue;
+			}
+			bool bCancelCompressedBlock = true;
+			for (FFileIoStoreBlockScatter& Scatter : CompressedBlock->ScatterList)
+			{
+				if (Scatter.Size > 0 && Scatter.Request != &ResolvedRequest)
+				{
+					bCancelCompressedBlock = false;
+					bCancelReadRequest = false;
+				}
+				else
+				{
+					Scatter.Size = 0;
+				}
+			}
+			if (bCancelCompressedBlock)
+			{
+				CompressedBlock->bCancelled = true;
+				CompressedBlocksMap.Remove(CompressedBlock->Key);
+			}
+		}
+		if (bCancelReadRequest)
+		{
+			if (!ReadRequest.ImmediateScatter.Request)
+			{
+				RawBlocksMap.Remove(ReadRequest.Key);
+			}
+			ReadRequest.bCancelled = true;
+#if DO_CHECK
+			for (FFileIoStoreCompressedBlock* CompressedBlock : ReadRequest.CompressedBlocks)
+			{
+				check(CompressedBlock->bCancelled);
+				for (FFileIoStoreBlockScatter& Scatter : CompressedBlock->ScatterList)
+				{
+					check(!Scatter.Request->DispatcherRequest || Scatter.Request->DispatcherRequest->IsCancelled());
+				}
+			}
+#endif
+		}
+	}
+	RequestQueue.Unlock();
+
+	return bShouldComplete;
+}
+
+void FFileIoStoreRequestTracker::UpdatePriorityForIoRequest(FFileIoStoreResolvedRequest& ResolvedRequest)
+{
+	TRACE_CPUPROFILER_EVENT_SCOPE(RequestTrackerUpdatePriorityForIoRequest);
+	bool bUpdateOrder = false;
+	FFileIoStoreReadRequestLink* Link = ResolvedRequest.ReadRequestsHead;
+	while (Link)
+	{
+		FFileIoStoreReadRequest& ReadRequest = Link->ReadRequest;
+		Link = Link->Next;
+		if (ResolvedRequest.GetPriority() > ReadRequest.Priority)
+		{
+			ReadRequest.Priority = ResolvedRequest.GetPriority();
+			bUpdateOrder = true;
+		}
+	}
+	if (bUpdateOrder)
+	{
+		RequestQueue.UpdateOrder();
+	}
+}
+
+void FFileIoStoreRequestTracker::ReleaseIoRequestReferences(FFileIoStoreResolvedRequest& ResolvedRequest)
+{
+	FFileIoStoreReadRequestLink* Link = ResolvedRequest.ReadRequestsHead;
+	while (Link)
+	{
+		FFileIoStoreReadRequestLink* Next = Link->Next;
+		check(Link->ReadRequest.RefCount > 0);
+		if (--Link->ReadRequest.RefCount == 0)
+		{
+			for (FFileIoStoreCompressedBlock* CompressedBlock : Link->ReadRequest.CompressedBlocks)
+			{
+				check(CompressedBlock->RefCount > 0);
+				if (--CompressedBlock->RefCount == 0)
+				{
+					RequestAllocator.Free(CompressedBlock);
+				}
+			}
+			RequestAllocator.Free(&Link->ReadRequest);
+		}
+		RequestAllocator.Free(Link);
+		Link = Next;
+	}
+	ResolvedRequest.ReadRequestsHead = nullptr;
+	ResolvedRequest.ReadRequestsTail = nullptr;
+	RequestAllocator.Free(&ResolvedRequest);
+}
+
+int64 FFileIoStoreRequestTracker::GetLiveReadRequestsCount() const
+{
+	return RequestAllocator.GetLiveReadRequestsCount();
+}
+
+FFileIoStore::FFileIoStore(TUniquePtr<IPlatformFileIoStore>&& InPlatformImpl)
+	: BlockCache(Stats)
+	, BufferAllocator(Stats)
+	, RequestTracker(RequestAllocator, RequestQueue)
+	, PlatformImpl(MoveTemp(InPlatformImpl))
+{
+}
+
+FFileIoStore::~FFileIoStore()
+{
+	delete Thread;
+}
+
+void FFileIoStore::Initialize(TSharedRef<const FIoDispatcherBackendContext> InContext)
+{
+	check(!Thread);
+
+	BackendContext = InContext;
+	bIsMultithreaded = InContext->bIsMultiThreaded;
+
+	ReadBufferSize = (GIoDispatcherBufferSizeKB > 0 ? uint64(GIoDispatcherBufferSizeKB) << 10 : 256 << 10);
+
+	uint64 BufferMemorySize = uint64(GIoDispatcherBufferMemoryMB) << 20ull;
+	uint64 BufferSize = uint64(GIoDispatcherBufferSizeKB) << 10ull;
+	uint32 BufferAlignment = uint32(GIoDispatcherBufferAlignment);
+	BufferAllocator.Initialize(BufferMemorySize, BufferSize, BufferAlignment);
+
+	uint64 CacheMemorySize = uint64(GIoDispatcherCacheSizeMB) << 20ull;
+	BlockCache.Initialize(CacheMemorySize, BufferSize);
+
+	PlatformImpl->Initialize({
+		&BackendContext->WakeUpDispatcherThreadDelegate,
+		&RequestAllocator,
+		&BufferAllocator,
+		&BlockCache,
+		&Stats
+	});
+
+	uint64 DecompressionContextCount = uint64(GIoDispatcherDecompressionWorkerCount > 0 ? GIoDispatcherDecompressionWorkerCount : 4);
+	for (uint64 ContextIndex = 0; ContextIndex < DecompressionContextCount; ++ContextIndex)
+	{
+		FFileIoStoreCompressionContext* Context = new FFileIoStoreCompressionContext();
+		Context->Next = FirstFreeCompressionContext;
+		FirstFreeCompressionContext = Context;
+	}
+
+	Thread = FRunnableThread::Create(this, TEXT("IoService"), 0, TPri_AboveNormal);
+}
+
+TIoStatusOr<FIoContainerHeader> FFileIoStore::Mount(const TCHAR* InTocPath, int32 Order, const FGuid& EncryptionKeyGuid, const FAES::FAESKey& EncryptionKey)
+{
+	TUniquePtr<FFileIoStoreReader> Reader(new FFileIoStoreReader(*PlatformImpl, Stats));
+	FIoStatus IoStatus = Reader->Initialize(InTocPath, Order);
+	if (!IoStatus.IsOk())
+	{
+		return IoStatus;
+	}
+
+	if (Reader->IsEncrypted())
+	{
+		if (Reader->GetEncryptionKeyGuid() == EncryptionKeyGuid && EncryptionKey.IsValid())
+		{
+			Reader->SetEncryptionKey(EncryptionKey);
+		}
+		else
+		{
+			return FIoStatus(EIoErrorCode::InvalidEncryptionKey, *FString::Printf(TEXT("Invalid encryption key '%s' (container '%s', encryption key '%s')"),
+				*EncryptionKeyGuid.ToString(), InTocPath, *Reader->GetEncryptionKeyGuid().ToString()));
+		}
+	}
+
+	TIoStatusOr<FIoContainerHeader> ContainerHeaderReadResult = Reader->ReadContainerHeader();
+	FIoContainerHeader ContainerHeader;
+	if (ContainerHeaderReadResult.IsOk())
+	{
+		ContainerHeader = ContainerHeaderReadResult.ConsumeValueOrDie();
+	}
+	else if (EIoErrorCode::NotFound != ContainerHeaderReadResult.Status().GetErrorCode())
+	{
+		return ContainerHeaderReadResult;
+	}
+
+	int32 InsertionIndex;
+	{
+		FWriteScopeLock _(IoStoreReadersLock);
+		InsertionIndex = Algo::UpperBound(IoStoreReaders, Reader, [](const TUniquePtr<FFileIoStoreReader>& A, const TUniquePtr<FFileIoStoreReader>& B)
+		{
+			if (A->GetOrder() != B->GetOrder())
+			{
+				return A->GetOrder() > B->GetOrder();
+			}
+			return A->GetContainerInstanceId() > B->GetContainerInstanceId();
+		});
+		IoStoreReaders.Insert(MoveTemp(Reader), InsertionIndex);
+		UE_LOG(LogIoDispatcher, Display, TEXT("Mounting container '%s' in location slot %d"), InTocPath, InsertionIndex);
+	}
+
+	return ContainerHeader;
+}
+
+bool FFileIoStore::Unmount(const TCHAR* InTocPath)
+{
+	FWriteScopeLock _(IoStoreReadersLock);
+	
+	for (int32 Idx = 0; Idx < IoStoreReaders.Num(); ++Idx)
+	{
+		if (IoStoreReaders[Idx]->GetContainerFile().FilePath == InTocPath)
+		{
+			UE_LOG(LogIoDispatcher, Display, TEXT("Unmounting container '%s'"), InTocPath);
+
+			// Cancel pending I/O requests trying to read from the container
+			for (const FFileIoStoreContainerFilePartition& Partition : IoStoreReaders[Idx]->GetContainerFile().Partitions)
+			{
+				RequestQueue.CancelRequestsWithFileHandle(Partition.FileHandle);
+			}
+
+			FIoContainerId ContainerId = IoStoreReaders[Idx]->GetContainerId();
+			IoStoreReaders.RemoveAt(Idx);
+			
+			return true;
+		}
+	}
+
+	UE_LOG(LogIoDispatcher, Display, TEXT("Failed to unmount container '%s'"), InTocPath);
+	
+	return false;
+}
+
+bool FFileIoStore::Resolve(FIoRequestImpl* Request)
+{
+	FReadScopeLock _(IoStoreReadersLock);
+	for (const TUniquePtr<FFileIoStoreReader>& Reader : IoStoreReaders)
+	{
+		if (const FIoOffsetAndLength* OffsetAndLength = Reader->Resolve(Request->ChunkId))
+		{
+			uint64 RequestedOffset = Request->Options.GetOffset();
+			uint64 ResolvedOffset = OffsetAndLength->GetOffset() + RequestedOffset;
+			uint64 ResolvedSize = 0;
+			if (RequestedOffset <= OffsetAndLength->GetLength())
+			{
+				ResolvedSize = FMath::Min(Request->Options.GetSize(), OffsetAndLength->GetLength() - RequestedOffset);
+			}
+
+			FFileIoStoreResolvedRequest* ResolvedRequest = RequestAllocator.AllocResolvedRequest(
+				*Request,
+				Reader->GetContainerFile(),
+				ResolvedOffset,
+				ResolvedSize);
+			Request->BackendData = ResolvedRequest;
+
+			if (ResolvedSize > 0)
+			{
+				FFileIoStoreReadRequestList CustomRequests;
+				if (PlatformImpl->CreateCustomRequests(*ResolvedRequest, CustomRequests))
+				{
+					Stats.OnReadRequestsQueued(CustomRequests);
+					RequestTracker.AddReadRequestsToResolvedRequest(CustomRequests, *ResolvedRequest);
+					RequestQueue.Push(CustomRequests);
+					OnNewPendingRequestsAdded();
+				}
+				else
+				{
+					ReadBlocks(*ResolvedRequest);
+				}
+			}
+			else
+			{
+				// Nothing to read
+				CompleteDispatcherRequest(ResolvedRequest);
+				RequestTracker.ReleaseIoRequestReferences(*ResolvedRequest);
+			}
+
+			return true;
+		}
+	}
+
+	return false;
+}
+
+void FFileIoStore::CancelIoRequest(FIoRequestImpl* Request)
+{
+	if (Request->BackendData)
+	{
+		FFileIoStoreResolvedRequest* ResolvedRequest = static_cast<FFileIoStoreResolvedRequest*>(Request->BackendData);
+		bool bShouldComplete = RequestTracker.CancelIoRequest(*ResolvedRequest);
+		if (bShouldComplete)
+		{
+			ResolvedRequest->bCancelled = true;
+			CompleteDispatcherRequest(ResolvedRequest);
+		}
+	}
+}
+
+void FFileIoStore::UpdatePriorityForIoRequest(FIoRequestImpl* Request)
+{
+	if (Request->BackendData)
+	{
+		FFileIoStoreResolvedRequest* ResolvedRequest = static_cast<FFileIoStoreResolvedRequest*>(Request->BackendData);
+		RequestTracker.UpdatePriorityForIoRequest(*ResolvedRequest);
+	}
+}
+
+bool FFileIoStore::DoesChunkExist(const FIoChunkId& ChunkId) const
+{
+	FReadScopeLock _(IoStoreReadersLock);
+	for (const TUniquePtr<FFileIoStoreReader>& Reader : IoStoreReaders)
+	{
+		if (Reader->DoesChunkExist(ChunkId))
+		{
+			return true;
+		}
+	}
+	return false;
+}
+
+TIoStatusOr<uint64> FFileIoStore::GetSizeForChunk(const FIoChunkId& ChunkId) const
+{
+	FReadScopeLock _(IoStoreReadersLock);
+	for (const TUniquePtr<FFileIoStoreReader>& Reader : IoStoreReaders)
+	{
+		TIoStatusOr<uint64> ReaderResult = Reader->GetSizeForChunk(ChunkId);
+		if (ReaderResult.IsOk())
+		{
+			return ReaderResult;
+		}
+	}
+	return FIoStatus(EIoErrorCode::NotFound);
+}
+
+FAutoConsoleTaskPriority CPrio_IoDispatcherTaskPriority(
+	TEXT("TaskGraph.TaskPriorities.IoDispatcherAsyncTasks"),
+	TEXT("Task and thread priority for IoDispatcher decompression."),
+	ENamedThreads::BackgroundThreadPriority, // if we have background priority task threads, then use them...
+	ENamedThreads::NormalTaskPriority, // .. at normal task priority
+	ENamedThreads::NormalTaskPriority // if we don't have background threads, then use normal priority threads at normal task priority instead
+);
+
+ENamedThreads::Type FFileIoStore::FDecompressAsyncTask::GetDesiredThread()
+{
+	return CPrio_IoDispatcherTaskPriority.Get();
+}
+
+void FFileIoStore::ScatterBlock(FFileIoStoreCompressedBlock* CompressedBlock, bool bIsAsync)
+{
+	LLM_SCOPE(ELLMTag::FileSystem);
+	TRACE_CPUPROFILER_EVENT_SCOPE(IoDispatcherScatter);
+	
+	check(!CompressedBlock->bFailed);
+
+	FFileIoStoreCompressionContext* CompressionContext = CompressedBlock->CompressionContext;
+	check(CompressionContext);
+	uint8* CompressedBuffer;
+	if (CompressedBlock->RawBlocks.Num() > 1)
+	{
+		check(CompressedBlock->CompressedDataBuffer);
+		CompressedBuffer = CompressedBlock->CompressedDataBuffer;
+	}
+	else
+	{
+		FFileIoStoreReadRequest* RawBlock = CompressedBlock->RawBlocks[0];
+		check(CompressedBlock->RawOffset >= RawBlock->Offset);
+		uint64 OffsetInBuffer = CompressedBlock->RawOffset - RawBlock->Offset;
+		CompressedBuffer = RawBlock->Buffer->Memory + OffsetInBuffer;
+	}
+	if (CompressedBlock->SignatureHash)
+	{
+		FSHAHash BlockHash;
+		FSHA1::HashBuffer(CompressedBuffer, CompressedBlock->RawSize, BlockHash.Hash);
+		if (*CompressedBlock->SignatureHash != BlockHash)
+		{
+			FIoSignatureError Error;
+			{
+				FReadScopeLock _(IoStoreReadersLock);
+				for (const TUniquePtr<FFileIoStoreReader>& Reader : IoStoreReaders)
+				{
+					if (CompressedBlock->Key.FileIndex == Reader->GetContainerInstanceId())
+					{
+						Error.ContainerName = FPaths::GetBaseFilename(Reader->GetContainerFile().FilePath);
+					}
+				}
+				Error.BlockIndex = CompressedBlock->Key.BlockIndex;
+				Error.ExpectedHash = *CompressedBlock->SignatureHash;
+				Error.ActualHash = BlockHash;
+			}
+
+			UE_LOG(LogIoDispatcher, Warning, TEXT("Signature error detected in container '%s' at block index '%d'"), *Error.ContainerName, Error.BlockIndex);
+
+			check(BackendContext);
+			if (BackendContext->SignatureErrorDelegate.IsBound())
+			{
+				BackendContext->SignatureErrorDelegate.Broadcast(Error);
+			}
+		}
+	}
+	if (!CompressedBlock->bFailed)
+	{
+		if (CompressedBlock->EncryptionKey.IsValid())
+		{
+			FAES::DecryptData(CompressedBuffer, CompressedBlock->RawSize, CompressedBlock->EncryptionKey);
+		}
+		uint8* UncompressedBuffer;
+		if (CompressedBlock->CompressionMethod.IsNone())
+		{
+			UncompressedBuffer = CompressedBuffer;
+		}
+		else
+		{
+			if (CompressionContext->UncompressedBufferSize < CompressedBlock->UncompressedSize)
+			{
+				FMemory::Free(CompressionContext->UncompressedBuffer);
+				CompressionContext->UncompressedBuffer = reinterpret_cast<uint8*>(FMemory::Malloc(CompressedBlock->UncompressedSize));
+				CompressionContext->UncompressedBufferSize = CompressedBlock->UncompressedSize;
+			}
+			UncompressedBuffer = CompressionContext->UncompressedBuffer;
+
+			bool bFailed = !FCompression::UncompressMemory(CompressedBlock->CompressionMethod, UncompressedBuffer, int32(CompressedBlock->UncompressedSize), CompressedBuffer, int32(CompressedBlock->CompressedSize));
+			if (bFailed)
+			{
+				UE_LOG(LogIoDispatcher, Warning, TEXT("Failed decompressing block"));
+				CompressedBlock->bFailed = true;
+			}
+		}
+
+		for (FFileIoStoreBlockScatter& Scatter : CompressedBlock->ScatterList)
+		{
+			if (Scatter.Size)
+			{
+				check(Scatter.DstOffset + Scatter.Size <= Scatter.Request->GetBuffer().DataSize());
+				check(Scatter.SrcOffset + Scatter.Size <= CompressedBlock->UncompressedSize);
+				FMemory::Memcpy(Scatter.Request->GetBuffer().Data() + Scatter.DstOffset, UncompressedBuffer + Scatter.SrcOffset, Scatter.Size);
+			}
+		}
+	}
+
+	if (bIsAsync)
+	{
+		FScopeLock Lock(&DecompressedBlocksCritical);
+		CompressedBlock->Next = FirstDecompressedBlock;
+		FirstDecompressedBlock = CompressedBlock;
+
+		BackendContext->WakeUpDispatcherThreadDelegate.Execute();
+	}
+}
+
+void FFileIoStore::CompleteDispatcherRequest(FFileIoStoreResolvedRequest* ResolvedRequest)
+{
+	check(ResolvedRequest);
+	check(ResolvedRequest->DispatcherRequest);
+	FIoRequestImpl* DispatcherRequest = ResolvedRequest->DispatcherRequest;
+	ResolvedRequest->DispatcherRequest = nullptr;
+	if (ResolvedRequest->bFailed)
+	{
+		DispatcherRequest->SetFailed();
+	}
+	DispatcherRequest->BackendData = nullptr;
+	if (!CompletedRequestsTail)
+	{
+		CompletedRequestsHead = CompletedRequestsTail = DispatcherRequest;
+	}
+	else
+	{
+		CompletedRequestsTail->NextRequest = DispatcherRequest;
+		CompletedRequestsTail = DispatcherRequest;
+	}
+	CompletedRequestsTail->NextRequest = nullptr;
+}
+
+void FFileIoStore::FinalizeCompressedBlock(FFileIoStoreCompressedBlock* CompressedBlock)
+{
+	Stats.OnDecompressComplete(CompressedBlock); 
+
+	if (CompressedBlock->RawBlocks.Num() > 1)
+	{
+		check(CompressedBlock->CompressedDataBuffer || CompressedBlock->bCancelled || CompressedBlock->bFailed);
+		if (CompressedBlock->CompressedDataBuffer)
+		{
+			FMemory::Free(CompressedBlock->CompressedDataBuffer);
+		}
+	}
+	else
+	{
+		FFileIoStoreReadRequest* RawBlock = CompressedBlock->RawBlocks[0];
+		check(RawBlock->BufferRefCount > 0);
+		if (--RawBlock->BufferRefCount == 0)
+		{
+			check(RawBlock->Buffer || RawBlock->bCancelled);
+			if (RawBlock->Buffer)
+			{
+				FreeBuffer(*RawBlock->Buffer);
+				RawBlock->Buffer = nullptr;
+			}
+		}
+	}
+	check(CompressedBlock->CompressionContext || CompressedBlock->bCancelled || CompressedBlock->bFailed);
+	if (CompressedBlock->CompressionContext)
+	{
+		FreeCompressionContext(CompressedBlock->CompressionContext);
+	}
+	for (int32 ScatterIndex = 0, ScatterCount = CompressedBlock->ScatterList.Num(); ScatterIndex < ScatterCount; ++ScatterIndex)
+	{
+		FFileIoStoreBlockScatter& Scatter = CompressedBlock->ScatterList[ScatterIndex];
+		Stats.OnBytesScattered(Scatter.Size);
+		Scatter.Request->bFailed |= CompressedBlock->bFailed;
+		check(!CompressedBlock->bCancelled || !Scatter.Request->DispatcherRequest || Scatter.Request->DispatcherRequest->IsCancelled());
+		check(Scatter.Request->UnfinishedReadsCount > 0);
+		if (--Scatter.Request->UnfinishedReadsCount == 0)
+		{
+			if (!Scatter.Request->bCancelled)
+			{
+				CompleteDispatcherRequest(Scatter.Request);
+			}
+			RequestTracker.ReleaseIoRequestReferences(*Scatter.Request);
+		}
+	}
+}
+
+FIoRequestImpl* FFileIoStore::GetCompletedRequests()
+{
+	LLM_SCOPE(ELLMTag::FileSystem);
+	//TRACE_CPUPROFILER_EVENT_SCOPE(GetCompletedRequests);
+	
+	if (!bIsMultithreaded)
+	{
+		while (PlatformImpl->StartRequests(RequestQueue));
+	}
+
+	FFileIoStoreReadRequestList CompletedRequests;
+	PlatformImpl->GetCompletedRequests(CompletedRequests);
+	Stats.OnReadRequestsCompleted(CompletedRequests);
+	for (auto It = CompletedRequests.Steal(); It; ++It)
+	{
+		FFileIoStoreReadRequest* CompletedRequest = *It;
+
+		if (!CompletedRequest->ImmediateScatter.Request)
+		{
+			check(CompletedRequest->Buffer || CompletedRequest->bCancelled);
+			RequestTracker.RemoveRawBlock(CompletedRequest);
+			
+			//TRACE_CPUPROFILER_EVENT_SCOPE(ProcessCompletedBlock);
+			for (FFileIoStoreCompressedBlock* CompressedBlock : CompletedRequest->CompressedBlocks)
+			{
+				CompressedBlock->bFailed |= CompletedRequest->bFailed;
+				CompressedBlock->bCancelled |= CompletedRequest->bCancelled;
+				if (CompressedBlock->RawBlocks.Num() > 1)
+				{
+					//TRACE_CPUPROFILER_EVENT_SCOPE(HandleComplexBlock);
+					if (!(CompressedBlock->bCancelled | CompressedBlock->bFailed))
+					{
+						check(CompletedRequest->Buffer);
+						if (!CompressedBlock->CompressedDataBuffer)
+						{
+							CompressedBlock->CompressedDataBuffer = reinterpret_cast<uint8*>(FMemory::Malloc(CompressedBlock->RawSize));
+						}
+
+						uint8* Src = CompletedRequest->Buffer->Memory;
+						uint8* Dst = CompressedBlock->CompressedDataBuffer;
+						uint64 CopySize = CompletedRequest->Size;
+						int64 CompletedBlockOffsetInBuffer = int64(CompletedRequest->Offset) - int64(CompressedBlock->RawOffset);
+						if (CompletedBlockOffsetInBuffer < 0)
+						{
+							Src -= CompletedBlockOffsetInBuffer;
+							CopySize += CompletedBlockOffsetInBuffer;
+						}
+						else
+						{
+							Dst += CompletedBlockOffsetInBuffer;
+						}
+						uint64 CompressedBlockRawEndOffset = CompressedBlock->RawOffset + CompressedBlock->RawSize;
+						uint64 CompletedBlockEndOffset = CompletedRequest->Offset + CompletedRequest->Size;
+						if (CompletedBlockEndOffset > CompressedBlockRawEndOffset)
+						{
+							CopySize -= CompletedBlockEndOffset - CompressedBlockRawEndOffset;
+						}
+						FMemory::Memcpy(Dst, Src, CopySize);
+					}
+					check(CompletedRequest->BufferRefCount > 0);
+					if (--CompletedRequest->BufferRefCount == 0)
+					{
+						if (CompletedRequest->Buffer)
+						{
+							FreeBuffer(*CompletedRequest->Buffer);
+							CompletedRequest->Buffer = nullptr;
+						}
+					}
+				}
+
+				check(CompressedBlock->UnfinishedRawBlocksCount > 0);
+				if (--CompressedBlock->UnfinishedRawBlocksCount == 0)
+				{
+					Stats.OnDecompressQueued(CompressedBlock);
+					RequestTracker.RemoveCompressedBlock(CompressedBlock);
+					if (!ReadyForDecompressionTail)
+					{
+						ReadyForDecompressionHead = ReadyForDecompressionTail = CompressedBlock;
+					}
+					else
+					{
+						ReadyForDecompressionTail->Next = CompressedBlock;
+						ReadyForDecompressionTail = CompressedBlock;
+					}
+					CompressedBlock->Next = nullptr;
+				}
+			}
+		}
+		else
+		{
+			check(!CompletedRequest->Buffer);
+			Stats.OnBytesScattered(CompletedRequest->ImmediateScatter.Size);
+			FFileIoStoreResolvedRequest* CompletedResolvedRequest = CompletedRequest->ImmediateScatter.Request;
+			CompletedResolvedRequest->bFailed |= CompletedRequest->bFailed;
+			check(!CompletedRequest->bCancelled || !CompletedResolvedRequest->DispatcherRequest || CompletedResolvedRequest->DispatcherRequest->IsCancelled());
+			check(CompletedResolvedRequest->UnfinishedReadsCount > 0);
+			if (--CompletedResolvedRequest->UnfinishedReadsCount == 0)
+			{
+				if (!CompletedResolvedRequest->bCancelled)
+				{
+					CompleteDispatcherRequest(CompletedResolvedRequest);
+				}
+				RequestTracker.ReleaseIoRequestReferences(*CompletedResolvedRequest);
+			}
+		}
+	}
+	
+	FFileIoStoreCompressedBlock* BlockToReap;
+	{
+		FScopeLock Lock(&DecompressedBlocksCritical);
+		BlockToReap = FirstDecompressedBlock;
+		FirstDecompressedBlock = nullptr;
+	}
+
+	while (BlockToReap)
+	{
+		FFileIoStoreCompressedBlock* Next = BlockToReap->Next;
+		FinalizeCompressedBlock(BlockToReap);
+		BlockToReap = Next;
+	}
+
+	FFileIoStoreCompressedBlock* BlockToDecompress = ReadyForDecompressionHead;
+	while (BlockToDecompress)
+	{
+		FFileIoStoreCompressedBlock* Next = BlockToDecompress->Next;
+		if (BlockToDecompress->bFailed | BlockToDecompress->bCancelled)
+		{
+			FinalizeCompressedBlock(BlockToDecompress);
+			BlockToDecompress = Next;
+			continue;
+		}
+		
+		BlockToDecompress->CompressionContext = AllocCompressionContext();
+		if (!BlockToDecompress->CompressionContext)
+		{
+			break;
+		}
+
+		for (const FFileIoStoreBlockScatter& Scatter : BlockToDecompress->ScatterList)
+		{
+			if (Scatter.Size)
+			{
+				FIoRequestImpl* DispatcherRequest = Scatter.Request->DispatcherRequest;
+				check(DispatcherRequest);
+				if (!DispatcherRequest->HasBuffer())
+				{
+					DispatcherRequest->CreateBuffer(Scatter.Request->ResolvedSize);
+				}
+			}
+		}
+
+		// Scatter block asynchronous when the block is compressed, encrypted or signed
+		const bool bScatterAsync = bIsMultithreaded && (!BlockToDecompress->CompressionMethod.IsNone() || BlockToDecompress->EncryptionKey.IsValid() || BlockToDecompress->SignatureHash);
+		if (bScatterAsync)
+		{
+			TGraphTask<FDecompressAsyncTask>::CreateTask().ConstructAndDispatchWhenReady(*this, BlockToDecompress);
+		}
+		else
+		{
+			ScatterBlock(BlockToDecompress, false);
+			FinalizeCompressedBlock(BlockToDecompress);
+		}
+		BlockToDecompress = Next;
+	}
+	ReadyForDecompressionHead = BlockToDecompress;
+	if (!ReadyForDecompressionHead)
+	{
+		ReadyForDecompressionTail = nullptr;
+	}
+
+	FIoRequestImpl* Result = CompletedRequestsHead;
+	CompletedRequestsHead = CompletedRequestsTail = nullptr;
+	return Result;
+}
+
+TIoStatusOr<FIoMappedRegion> FFileIoStore::OpenMapped(const FIoChunkId& ChunkId, const FIoReadOptions& Options)
+{
+	if (!FPlatformProperties::SupportsMemoryMappedFiles())
+	{
+		return FIoStatus(EIoErrorCode::Unknown, TEXT("Platform does not support memory mapped files"));
+	}
+
+	if (Options.GetTargetVa() != nullptr)
+	{
+		return FIoStatus(EIoErrorCode::InvalidParameter, TEXT("Invalid read options"));
+	}
+
+	IPlatformFile& Ipf = FPlatformFileManager::Get().GetPlatformFile();
+
+	FReadScopeLock _(IoStoreReadersLock);
+	for (TUniquePtr<FFileIoStoreReader>& Reader : IoStoreReaders)
+	{
+		if (const FIoOffsetAndLength* OffsetAndLength = Reader->Resolve(ChunkId))
+		{
+			uint64 ResolvedOffset = OffsetAndLength->GetOffset();
+			uint64 ResolvedSize = FMath::Min(Options.GetSize(), OffsetAndLength->GetLength());
+			
+			const FFileIoStoreContainerFile& ContainerFile = Reader->GetContainerFile();
+			
+			int32 BlockIndex = int32(ResolvedOffset / ContainerFile.CompressionBlockSize);
+			const FIoStoreTocCompressedBlockEntry& CompressionBlockEntry = ContainerFile.CompressionBlocks[BlockIndex];
+			const int64 BlockOffset = (int64)CompressionBlockEntry.GetOffset();
+			check(BlockOffset > 0 && IsAligned(BlockOffset, FPlatformProperties::GetMemoryMappingAlignment()));
+
+			IMappedFileHandle* MappedFileHandle = Reader->GetMappedContainerFileHandle(BlockOffset);
+			IMappedFileRegion* MappedFileRegion = MappedFileHandle->MapRegion(BlockOffset + Options.GetOffset(), ResolvedSize);
+			if (MappedFileRegion != nullptr)
+			{
+				check(IsAligned(MappedFileRegion->GetMappedPtr(), FPlatformProperties::GetMemoryMappingAlignment()));
+				return FIoMappedRegion{ MappedFileHandle, MappedFileRegion };
+			}
+			else
+			{
+				return FIoStatus(EIoErrorCode::ReadError);
+			}
+		}
+	}
+
+	// We didn't find any entry for the ChunkId.
+	return FIoStatus(EIoErrorCode::NotFound);
+}
+
+void FFileIoStore::ReopenAllFileHandles()
+{
+	UE_CLOG(RequestTracker.GetLiveReadRequestsCount(), LogIoDispatcher, Warning, TEXT("Calling ReopenAllFileHandles with read requests in flight"));
+	FWriteScopeLock _(IoStoreReadersLock);
+	for (const TUniquePtr<FFileIoStoreReader>& Reader : IoStoreReaders)
+	{
+		Reader->ReopenAllFileHandles();
+	}
+}
+
+void FFileIoStore::OnNewPendingRequestsAdded()
+{
+	if (bIsMultithreaded)
+	{
+		PlatformImpl->ServiceNotify();
+	}
+}
+
+void FFileIoStore::ReadBlocks(FFileIoStoreResolvedRequest& ResolvedRequest)
+{
+	/*TStringBuilder<256> ScopeName;
+	ScopeName.Appendf(TEXT("ReadBlock %d"), BlockIndex);
+	TRACE_CPUPROFILER_EVENT_SCOPE_TEXT(*ScopeName);*/
+
+	const FFileIoStoreContainerFile& ContainerFile = ResolvedRequest.GetContainerFile();
+	const uint64 CompressionBlockSize = ContainerFile.CompressionBlockSize;
+	const uint64 RequestEndOffset = ResolvedRequest.ResolvedOffset + ResolvedRequest.ResolvedSize;
+	int32 RequestBeginBlockIndex = int32(ResolvedRequest.ResolvedOffset / CompressionBlockSize);
+	int32 RequestEndBlockIndex = int32((RequestEndOffset - 1) / CompressionBlockSize);
+
+	FFileIoStoreReadRequestList NewBlocks;
+
+	uint64 RequestStartOffsetInBlock = ResolvedRequest.ResolvedOffset - RequestBeginBlockIndex * CompressionBlockSize;
+	uint64 RequestRemainingBytes = ResolvedRequest.ResolvedSize;
+	uint64 OffsetInRequest = 0;
+	for (int32 CompressedBlockIndex = RequestBeginBlockIndex; CompressedBlockIndex <= RequestEndBlockIndex; ++CompressedBlockIndex)
+	{
+		FFileIoStoreBlockKey CompressedBlockKey;
+		CompressedBlockKey.FileIndex = ResolvedRequest.GetContainerFile().ContainerInstanceId;
+		CompressedBlockKey.BlockIndex = CompressedBlockIndex;
+		bool bCompressedBlockWasAdded;
+		FFileIoStoreCompressedBlock* CompressedBlock = RequestTracker.FindOrAddCompressedBlock(CompressedBlockKey, bCompressedBlockWasAdded);
+		check(CompressedBlock);
+		check(!CompressedBlock->bCancelled);
+		if (bCompressedBlockWasAdded)
+		{
+			CompressedBlock->EncryptionKey = ContainerFile.EncryptionKey;
+			const FIoStoreTocCompressedBlockEntry& CompressionBlockEntry = ContainerFile.CompressionBlocks[CompressedBlockIndex];
+			CompressedBlock->UncompressedSize = CompressionBlockEntry.GetUncompressedSize();
+			CompressedBlock->CompressedSize = CompressionBlockEntry.GetCompressedSize();
+			CompressedBlock->CompressionMethod = ContainerFile.CompressionMethods[CompressionBlockEntry.GetCompressionMethodIndex()];
+			CompressedBlock->SignatureHash = EnumHasAnyFlags(ContainerFile.ContainerFlags, EIoContainerFlags::Signed) ? &ContainerFile.BlockSignatureHashes[CompressedBlockIndex] : nullptr;
+			CompressedBlock->RawSize = Align(CompressionBlockEntry.GetCompressedSize(), FAES::AESBlockSize); // The raw blocks size is always aligned to AES blocks size;
+
+			int32 PartitionIndex = int32(CompressionBlockEntry.GetOffset() / ContainerFile.PartitionSize);
+			const FFileIoStoreContainerFilePartition& Partition = ContainerFile.Partitions[PartitionIndex];
+			uint64 PartitionRawOffset = CompressionBlockEntry.GetOffset() % ContainerFile.PartitionSize;
+			CompressedBlock->RawOffset = PartitionRawOffset;
+			const uint32 RawBeginBlockIndex = uint32(PartitionRawOffset / ReadBufferSize);
+			const uint32 RawEndBlockIndex = uint32((PartitionRawOffset + CompressedBlock->RawSize - 1) / ReadBufferSize);
+			const uint32 RawBlockCount = RawEndBlockIndex - RawBeginBlockIndex + 1;
+			check(RawBlockCount > 0);
+			for (uint32 RawBlockIndex = RawBeginBlockIndex; RawBlockIndex <= RawEndBlockIndex; ++RawBlockIndex)
+			{
+				FFileIoStoreBlockKey RawBlockKey;
+				RawBlockKey.BlockIndex = RawBlockIndex;
+				RawBlockKey.FileIndex = Partition.ContainerFileIndex;
+
+				bool bRawBlockWasAdded;
+				FFileIoStoreReadRequest* RawBlock = RequestTracker.FindOrAddRawBlock(RawBlockKey, bRawBlockWasAdded);
+				check(RawBlock);
+				check(!RawBlock->bCancelled);
+				if (bRawBlockWasAdded)
+				{
+					RawBlock->Priority = ResolvedRequest.GetPriority();
+					RawBlock->FileHandle = Partition.FileHandle;
+					RawBlock->Offset = RawBlockIndex * ReadBufferSize;
+					uint64 ReadSize = FMath::Min(Partition.FileSize, RawBlock->Offset + ReadBufferSize) - RawBlock->Offset;
+					RawBlock->Size = ReadSize;
+					NewBlocks.Add(RawBlock);
+				}
+				RawBlock->BytesUsed += 
+					uint32(FMath::Min(CompressedBlock->RawOffset + CompressedBlock->RawSize, RawBlock->Offset + RawBlock->Size) -
+						   FMath::Max(CompressedBlock->RawOffset, RawBlock->Offset));
+				CompressedBlock->RawBlocks.Add(RawBlock);
+				++CompressedBlock->UnfinishedRawBlocksCount;
+				++CompressedBlock->RefCount;
+				RawBlock->CompressedBlocks.Add(CompressedBlock);
+				++RawBlock->BufferRefCount;
+			}
+		}
+		check(CompressedBlock->UncompressedSize > RequestStartOffsetInBlock);
+		uint64 RequestSizeInBlock = FMath::Min<uint64>(CompressedBlock->UncompressedSize - RequestStartOffsetInBlock, RequestRemainingBytes);
+		check(OffsetInRequest + RequestSizeInBlock <= ResolvedRequest.ResolvedSize);
+		check(RequestStartOffsetInBlock + RequestSizeInBlock <= CompressedBlock->UncompressedSize);
+
+		FFileIoStoreBlockScatter& Scatter = CompressedBlock->ScatterList.AddDefaulted_GetRef();
+		Scatter.Request = &ResolvedRequest;
+		Scatter.DstOffset = OffsetInRequest;
+		Scatter.SrcOffset = RequestStartOffsetInBlock;
+		Scatter.Size = RequestSizeInBlock;
+
+		RequestRemainingBytes -= RequestSizeInBlock;
+		OffsetInRequest += RequestSizeInBlock;
+		RequestStartOffsetInBlock = 0;
+
+		RequestTracker.AddReadRequestsToResolvedRequest(CompressedBlock, ResolvedRequest);
+	}
+
+	if (!NewBlocks.IsEmpty())
+	{
+		Stats.OnReadRequestsQueued(NewBlocks);
+		RequestQueue.Push(NewBlocks);
+		OnNewPendingRequestsAdded();
+	}
+}
+
+void FFileIoStore::FreeBuffer(FFileIoStoreBuffer& Buffer)
+{
+	BufferAllocator.FreeBuffer(&Buffer);
+	PlatformImpl->ServiceNotify();
+}
+
+FFileIoStoreCompressionContext* FFileIoStore::AllocCompressionContext()
+{
+	FFileIoStoreCompressionContext* Result = FirstFreeCompressionContext;
+	if (Result)
+	{
+		FirstFreeCompressionContext = FirstFreeCompressionContext->Next;
+	}
+	return Result;
+}
+
+void FFileIoStore::FreeCompressionContext(FFileIoStoreCompressionContext* CompressionContext)
+{
+	CompressionContext->Next = FirstFreeCompressionContext;
+	FirstFreeCompressionContext = CompressionContext;
+}
+
+void FFileIoStore::UpdateAsyncIOMinimumPriority()
+{
+	EAsyncIOPriorityAndFlags NewAsyncIOMinimumPriority = AIOP_MIN;
+	if (FFileIoStoreReadRequest* NextRequest = RequestQueue.Peek())
+	{
+		if (NextRequest->Priority >= IoDispatcherPriority_High)
+		{
+			NewAsyncIOMinimumPriority = AIOP_MAX;
+		}
+		else if (NextRequest->Priority >= IoDispatcherPriority_Medium)
+		{
+			NewAsyncIOMinimumPriority = AIOP_Normal;
+		}
+	}
+	if (NewAsyncIOMinimumPriority != CurrentAsyncIOMinimumPriority)
+	{
+		//TRACE_BOOKMARK(TEXT("SetAsyncMinimumPrioirity(%d)"), NewAsyncIOMinimumPriority);
+		FPlatformFileManager::Get().GetPlatformFile().SetAsyncMinimumPriority(NewAsyncIOMinimumPriority);
+		CurrentAsyncIOMinimumPriority = NewAsyncIOMinimumPriority;
+	}
+}
+
+bool FFileIoStore::Init()
+{
+	return true;
+}
+
+void FFileIoStore::Stop()
+{
+	bStopRequested = true;
+	PlatformImpl->ServiceNotify();
+}
+
+uint32 FFileIoStore::Run()
+{
+	while (!bStopRequested)
+	{
+		UpdateAsyncIOMinimumPriority();
+		if (!PlatformImpl->StartRequests(RequestQueue))
+		{
+			UpdateAsyncIOMinimumPriority();
+			PlatformImpl->ServiceWait();
+		}
+	}
+	return 0;
+}
+
+TSharedRef<FFileIoStore> CreateIoDispatcherFileBackend()
+{
+	bool bCheckForPlatformImplementation = true;
+	if (!FGenericPlatformProcess::SupportsMultithreading())
+	{
+		bCheckForPlatformImplementation = false;
+	}
+#if !UE_BUILD_SHIPPING
+	if (FParse::Param(FCommandLine::Get(), TEXT("forcegenericio")))
+	{
+		bCheckForPlatformImplementation = false;
+	}
+#endif
+
+	if (bCheckForPlatformImplementation)
+	{
+		if (FModuleManager::Get().ModuleExists(ANSI_TO_TCHAR(PLATFORM_IODISPATCHER_MODULE)))
+		{
+			IPlatformFileIoStoreModule* PlatformModule = FModuleManager::LoadModulePtr<IPlatformFileIoStoreModule>(ANSI_TO_TCHAR(PLATFORM_IODISPATCHER_MODULE));
+			if (PlatformModule)
+			{
+				TUniquePtr<IPlatformFileIoStore> PlatformImpl = PlatformModule->CreatePlatformFileIoStore();
+				if (PlatformImpl.IsValid())
+				{
+					return MakeShared<FFileIoStore>(MoveTemp(PlatformImpl));
+				}
+			}
+		}
+#if PLATFORM_IMPLEMENTS_IO
+		{
+			TUniquePtr<IPlatformFileIoStore> PlatformImpl = CreatePlatformFileIoStore();
+			if (PlatformImpl.IsValid())
+			{
+				return MakeShared<FFileIoStore>(MoveTemp(PlatformImpl));
+			}
+		}
+#endif
+	}
+	return MakeShared<FFileIoStore>(MakeUnique<FGenericFileIoStoreImpl>());
+}
+
+uint32 FFileIoStore::GetThreadId() const
+{
+	return Thread ? Thread->GetThreadID() : 0; 
+}
+
+CSV_DEFINE_CATEGORY(IoDispatcherFileBackend, true);
+CSV_DEFINE_CATEGORY(IoDispatcherFileBackendVerbose, false);
+
+// These stats go to both insights and csv by default
+// TODO: Ideally these should go to insights even if CSV is not capturing, but not be doubled-up if both CSV and Insights are capturing
+// TODO: It would also be nice to send these to insights as int64 without unit conversion where appropriate
+// IoDispatcher thread
+CSV_DEFINE_STAT(IoDispatcherFileBackend,			FrameBytesScatteredKB);
+CSV_DEFINE_STAT(IoDispatcherFileBackend,			QueuedFilesystemReadMB);
+CSV_DEFINE_STAT(IoDispatcherFileBackendVerbose,	QueuedFilesystemReads);
+CSV_DEFINE_STAT(IoDispatcherFileBackendVerbose,	QueuedUncompressBlocks);	
+CSV_DEFINE_STAT(IoDispatcherFileBackendVerbose,	QueuedUncompressInMB);		
+CSV_DEFINE_STAT(IoDispatcherFileBackendVerbose,	QueuedUncompressOutMB);	
+CSV_DEFINE_STAT(IoDispatcherFileBackendVerbose,	FrameBytesReadKB);				
+CSV_DEFINE_STAT(IoDispatcherFileBackendVerbose,	FrameBytesUncompressedInKB);	
+CSV_DEFINE_STAT(IoDispatcherFileBackendVerbose,	FrameBytesUncompressedOutKB);	
+
+// FileIoStore thread			
+CSV_DEFINE_STAT(IoDispatcherFileBackend,			FrameFilesystemBytesReadKB);	
+CSV_DEFINE_STAT(IoDispatcherFileBackend,			FrameSequentialReads);
+CSV_DEFINE_STAT(IoDispatcherFileBackend,			FrameSeeks);
+CSV_DEFINE_STAT(IoDispatcherFileBackendVerbose,	FrameFilesystemReads);
+CSV_DEFINE_STAT(IoDispatcherFileBackendVerbose,	FrameForwardSeeks);
+CSV_DEFINE_STAT(IoDispatcherFileBackendVerbose,	FrameBackwardSeeks);
+CSV_DEFINE_STAT(IoDispatcherFileBackendVerbose,	FrameHandleChangeSeeks);
+CSV_DEFINE_STAT(IoDispatcherFileBackendVerbose,	FrameSeekDistanceMB);	
+CSV_DEFINE_STAT(IoDispatcherFileBackendVerbose,	FrameBlockCacheStores);
+CSV_DEFINE_STAT(IoDispatcherFileBackendVerbose,	FrameBlockCacheStoresKB);
+CSV_DEFINE_STAT(IoDispatcherFileBackendVerbose,	FrameBlockCacheHits);	
+CSV_DEFINE_STAT(IoDispatcherFileBackendVerbose,	FrameBlockCacheHitKB);	
+CSV_DEFINE_STAT(IoDispatcherFileBackendVerbose,	FrameBlockCacheMisses);
+CSV_DEFINE_STAT(IoDispatcherFileBackendVerbose,	FrameBlockCacheMissKB);
+
+#if UE_FILEIOSTORE_STATS_ENABLED
+
+FFileIoStoreStats::FFileIoStoreStats()
+#if COUNTERSTRACE_ENABLED
+	: QueuedReadRequestsSizeCounter(TEXT("FileIoStore/QueuedReadRequestsSize"), TraceCounterDisplayHint_Memory)
+	, CompletedReadRequestsSizeCounter(TEXT("FileIoStore/CompletedReadRequestsSize"), TraceCounterDisplayHint_Memory)
+	, QueuedCompressedSizeCounter(TEXT("FileIoStore/QueuedCompressedSize"), TraceCounterDisplayHint_Memory)
+	, QueuedUncompressedSizeCounter(TEXT("FileIoStore/QueuedUncompressedSize"), TraceCounterDisplayHint_Memory)
+	, CompletedCompressedSizeCounter(TEXT("FileIoStore/CompletedCompressedSize"), TraceCounterDisplayHint_Memory)
+	, CompletedUncompressedSizeCounter(TEXT("FileIoStore/CompletedUncompressedSize"), TraceCounterDisplayHint_Memory)
+	, FileSystemSeeksTotalDistanceCounter(TEXT("FileIoStore/FileSystemSeeksTotalDistance"), TraceCounterDisplayHint_Memory)
+	, FileSystemSeeksForwardCountCounter(TEXT("FileIoStore/FileSystemSeeksForwardCount"), TraceCounterDisplayHint_None)
+	, FileSystemSeeksBackwardCountCounter(TEXT("FileIoStore/FileSystemSeeksBackwardCount"), TraceCounterDisplayHint_None)
+	, FileSystemSeeksChangeHandleCountCounter(TEXT("FileIoStore/FileSystemSeeksChangeHandleCount"), TraceCounterDisplayHint_None)
+	, FileSystemCompletedRequestsSizeCounter(TEXT("FileIoStore/FileSystemCompletedRequestsSize"), TraceCounterDisplayHint_Memory)
+	, BlockCacheStoredSizeCounter(TEXT("FileIoStore/BlockCacheStoredSize"), TraceCounterDisplayHint_Memory)
+	, BlockCacheHitSizeCounter(TEXT("FileIoStore/BlockCacheHitSize"), TraceCounterDisplayHint_Memory)
+	, BlockCacheMissedSizeCounter(TEXT("FileIoStore/BlockCacheMissedSize"), TraceCounterDisplayHint_Memory)
+	, ScatteredSizeCounter(TEXT("FileIoStore/ScatteredSize"), TraceCounterDisplayHint_Memory)
+	, TocMemoryCounter(TEXT("FileIoStore/TocMemory"), TraceCounterDisplayHint_Memory)
+	, AvailableBuffersCounter(TEXT("FileIoStore/AvailableBuffers"), TraceCounterDisplayHint_None)
+#endif
+{
+#if CSV_PROFILER
+	TickerHandle = FTSTicker::GetCoreTicker().AddTicker(FTickerDelegate::CreateRaw(this, &FFileIoStoreStats::CsvTick));
+#endif
+}
+
+FFileIoStoreStats::~FFileIoStoreStats()
+{
+#if CSV_PROFILER
+	FTSTicker::GetCoreTicker().RemoveTicker(TickerHandle);
+#endif
+}
+
+bool FFileIoStoreStats::CsvTick(float DeltaTime)
+{
+#if CSV_PROFILER
+	CSV_CUSTOM_STAT_DEFINED(QueuedFilesystemReadMB, BytesToApproxMB(QueuedFilesystemReadBytes), ECsvCustomStatOp::Set);
+	CSV_CUSTOM_STAT_DEFINED(QueuedFilesystemReads, (int32)QueuedFilesystemReads, ECsvCustomStatOp::Set);
+
+	CSV_CUSTOM_STAT_DEFINED(QueuedUncompressBlocks, int32(QueuedUncompressBlocks), ECsvCustomStatOp::Set);
+	CSV_CUSTOM_STAT_DEFINED(QueuedUncompressInMB, BytesToApproxMB(QueuedUncompressBytesIn), ECsvCustomStatOp::Set);
+	CSV_CUSTOM_STAT_DEFINED(QueuedUncompressOutMB, BytesToApproxMB(QueuedUncompressBytesOut), ECsvCustomStatOp::Set);
+#endif
+
+	return true;
+}
+
+void FFileIoStoreStats::OnReadRequestsQueued(const FFileIoStoreReadRequestList& Requests)
+{
+	uint64 TotalBytes = 0;
+	int32 NumReads = 0;
+	for (const FFileIoStoreReadRequest* Request : Requests)
+	{
+		++NumReads;
+		TotalBytes += Request->Size;
+	}
+
+#if CSV_PROFILER
+	QueuedFilesystemReadBytes += TotalBytes;
+	QueuedFilesystemReads += NumReads;
+#endif
+
+#if COUNTERSTRACE_ENABLED
+	QueuedReadRequestsSizeCounter.Add(TotalBytes);
+#endif
+}
+
+void FFileIoStoreStats::OnFilesystemReadStarted(const FFileIoStoreReadRequest* Request)
+{
+	CSV_CUSTOM_STAT_DEFINED(FrameFilesystemReads, 1, ECsvCustomStatOp::Accumulate);
+
+	if (LastHandle != Request->FileHandle)
+	{
+		OnHandleChangeSeek();
+	}
+	else if (LastOffset == Request->Offset)
+	{
+		OnSequentialRead();
+	}
+	else
+	{
+		OnSeek(LastOffset, Request->Offset);
+	}
+	LastOffset = Request->Offset + Request->Size;
+	LastHandle = Request->FileHandle;
+
+	CSV_CUSTOM_STAT_DEFINED(FrameFilesystemReads, 1, ECsvCustomStatOp::Accumulate);
+}
+
+void FFileIoStoreStats::OnFilesystemReadsStarted(const FFileIoStoreReadRequestList& Requests)
+{
+	int64 TotalBytes = 0;
+	int32 NumReads = 0;
+	for (const FFileIoStoreReadRequest* Request : Requests)
+	{
+		if (LastHandle != Request->FileHandle)
+		{
+			OnHandleChangeSeek();
+		}
+		else if (LastOffset == Request->Offset)
+		{
+			OnSequentialRead();
+		}
+		else
+		{
+			OnSeek(LastOffset, Request->Offset);
+		}
+
+		LastOffset = Request->Offset + Request->Size;
+		LastHandle = Request->FileHandle;
+
+		++NumReads;
+		TotalBytes += Request->Size;
+	}
+
+	CSV_CUSTOM_STAT_DEFINED(FrameFilesystemReads, NumReads, ECsvCustomStatOp::Accumulate);
+}
+
+void FFileIoStoreStats::OnFilesystemReadCompleted(const FFileIoStoreReadRequest* CompletedRequest)
+{
+	CSV_CUSTOM_STAT_DEFINED(FrameFilesystemBytesReadKB, BytesToApproxKB(CompletedRequest->Size), ECsvCustomStatOp::Accumulate);
+#if COUNTERSTRACE_ENABLED
+	FileSystemCompletedRequestsSizeCounter.Add(CompletedRequest->Size);
+#endif
+}
+
+void FFileIoStoreStats::OnFilesystemReadsCompleted(const FFileIoStoreReadRequestList& CompletedRequests)
+{
+	int64 TotalBytes = 0;
+	int32 NumReads = 0;
+	for (const FFileIoStoreReadRequest* Request : CompletedRequests)
+	{
+		++NumReads;
+		TotalBytes += Request->Size;
+	}
+
+	CSV_CUSTOM_STAT_DEFINED(FrameFilesystemBytesReadKB, BytesToApproxKB(TotalBytes), ECsvCustomStatOp::Accumulate);
+#if COUNTERSTRACE_ENABLED
+	FileSystemCompletedRequestsSizeCounter.Add(TotalBytes);
+#endif
+}
+
+void FFileIoStoreStats::OnReadRequestsCompleted(const FFileIoStoreReadRequestList& CompletedRequests)
+{
+	int64 TotalBytes = 0;
+	int32 NumReads = 0;
+	for (const FFileIoStoreReadRequest* Request : CompletedRequests)
+	{
+		++NumReads;
+		TotalBytes += Request->Size;
+	}
+
+#if CSV_PROFILER
+	QueuedFilesystemReadBytes -= TotalBytes;
+	QueuedFilesystemReads -= NumReads;
+#endif
+
+	CSV_CUSTOM_STAT_DEFINED(FrameBytesReadKB, BytesToApproxKB(TotalBytes), ECsvCustomStatOp::Accumulate); // TODO: Send to insights if CSV is disabled
+
+#if COUNTERSTRACE_ENABLED
+	QueuedReadRequestsSizeCounter.Subtract(TotalBytes);
+	CompletedReadRequestsSizeCounter.Add(TotalBytes);
+#endif
+}
+
+void FFileIoStoreStats::OnDecompressQueued(const FFileIoStoreCompressedBlock* CompressedBlock)
+{
+#if CSV_PROFILER
+	++QueuedUncompressBlocks;
+	QueuedUncompressBytesIn += CompressedBlock->CompressedSize;
+	QueuedUncompressBytesOut += CompressedBlock->UncompressedSize;
+#endif
+
+#if COUNTERSTRACE_ENABLED
+	QueuedCompressedSizeCounter.Add(CompressedBlock->CompressedSize);
+	QueuedUncompressedSizeCounter.Add(CompressedBlock->UncompressedSize);
+#endif
+}
+
+void FFileIoStoreStats::OnDecompressComplete(const FFileIoStoreCompressedBlock* CompressedBlock)
+{
+#if CSV_PROFILER
+	--QueuedUncompressBlocks;
+	QueuedUncompressBytesIn -= CompressedBlock->CompressedSize;
+	QueuedUncompressBytesOut -= CompressedBlock->UncompressedSize;
+#endif
+
+	CSV_CUSTOM_STAT_DEFINED(FrameBytesUncompressedInKB, BytesToApproxKB(CompressedBlock->CompressedSize), ECsvCustomStatOp::Accumulate); 
+	CSV_CUSTOM_STAT_DEFINED(FrameBytesUncompressedOutKB, BytesToApproxKB(CompressedBlock->UncompressedSize), ECsvCustomStatOp::Accumulate);
+
+#if COUNTERSTRACE_ENABLED
+	QueuedCompressedSizeCounter.Subtract(CompressedBlock->CompressedSize);
+	QueuedUncompressedSizeCounter.Subtract(CompressedBlock->UncompressedSize);
+
+	CompletedCompressedSizeCounter.Add(CompressedBlock->CompressedSize);
+	CompletedUncompressedSizeCounter.Add(CompressedBlock->UncompressedSize);
+#endif
+}
+
+void FFileIoStoreStats::OnBytesScattered(int64 NumBytes)
+{
+	CSV_CUSTOM_STAT_DEFINED(FrameBytesScatteredKB, BytesToApproxKB(NumBytes), ECsvCustomStatOp::Accumulate);
+#if COUNTERSTRACE_ENABLED
+	ScatteredSizeCounter.Add(NumBytes);
+#endif
+}
+
+void FFileIoStoreStats::OnSequentialRead()
+{
+	CSV_CUSTOM_STAT_DEFINED(FrameSequentialReads, 1, ECsvCustomStatOp::Accumulate);
+}
+
+void FFileIoStoreStats::OnSeek(uint64 PrevOffset, uint64 NewOffset)
+{
+	if (NewOffset > PrevOffset)
+	{
+		int64 Delta = NewOffset - PrevOffset;
+
+		CSV_CUSTOM_STAT_DEFINED(FrameForwardSeeks, 1, ECsvCustomStatOp::Accumulate);
+		CSV_CUSTOM_STAT_DEFINED(FrameSeekDistanceMB, BytesToApproxMB(Delta), ECsvCustomStatOp::Accumulate);
+
+#if COUNTERSTRACE_ENABLED
+		FileSystemSeeksTotalDistanceCounter.Add(Delta);
+		FileSystemSeeksForwardCountCounter.Increment();
+#endif
+	}
+	else
+	{
+		int64 Delta = PrevOffset - NewOffset;
+		CSV_CUSTOM_STAT_DEFINED(FrameBackwardSeeks, 1, ECsvCustomStatOp::Accumulate);
+		CSV_CUSTOM_STAT_DEFINED(FrameSeekDistanceMB, BytesToApproxMB(Delta), ECsvCustomStatOp::Accumulate);
+
+#if COUNTERSTRACE_ENABLED
+		FileSystemSeeksTotalDistanceCounter.Add(Delta);
+		FileSystemSeeksBackwardCountCounter.Increment();
+#endif
+	}
+
+	CSV_CUSTOM_STAT_DEFINED(FrameSeeks, 1, ECsvCustomStatOp::Accumulate);
+}
+
+void FFileIoStoreStats::OnHandleChangeSeek()
+{
+	CSV_CUSTOM_STAT_DEFINED(FrameHandleChangeSeeks, 1, ECsvCustomStatOp::Accumulate);
+	CSV_CUSTOM_STAT_DEFINED(FrameSeeks, 1, ECsvCustomStatOp::Accumulate);
+
+#if COUNTERSTRACE_ENABLED
+	FileSystemSeeksChangeHandleCountCounter.Increment();
+#endif
+}
+
+void FFileIoStoreStats::OnBlockCacheStore(uint64 NumBytes)
+{
+	CSV_CUSTOM_STAT_DEFINED(FrameBlockCacheStores, 1, ECsvCustomStatOp::Accumulate);
+	CSV_CUSTOM_STAT_DEFINED(FrameBlockCacheStoresKB, BytesToApproxKB(NumBytes), ECsvCustomStatOp::Accumulate);
+
+#if COUNTERSTRACE_ENABLED
+	BlockCacheStoredSizeCounter.Add(NumBytes);
+#endif
+}
+
+void FFileIoStoreStats::OnBlockCacheHit(uint64 NumBytes)
+{
+	CSV_CUSTOM_STAT_DEFINED(FrameBlockCacheHits, 1, ECsvCustomStatOp::Accumulate);
+	CSV_CUSTOM_STAT_DEFINED(FrameBlockCacheHitKB, BytesToApproxKB(NumBytes), ECsvCustomStatOp::Accumulate);
+
+#if COUNTERSTRACE_ENABLED
+	BlockCacheHitSizeCounter.Add(NumBytes);
+#endif
+}
+
+void FFileIoStoreStats::OnBlockCacheMiss(uint64 NumBytes)
+{
+	CSV_CUSTOM_STAT_DEFINED(FrameBlockCacheMisses, 1, ECsvCustomStatOp::Accumulate);
+	CSV_CUSTOM_STAT_DEFINED(FrameBlockCacheMissKB, BytesToApproxKB(NumBytes), ECsvCustomStatOp::Accumulate);
+
+#if COUNTERSTRACE_ENABLED
+	BlockCacheMissedSizeCounter.Add(NumBytes);
+#endif
+}
+
+void FFileIoStoreStats::OnTocMounted(uint64 AllocatedSize)
+{
+#if COUNTERSTRACE_ENABLED
+	TocMemoryCounter.Add(AllocatedSize);
+#endif
+}
+
+void FFileIoStoreStats::OnTocUnmounted(uint64 AllocatedSize)
+{
+#if COUNTERSTRACE_ENABLED
+	TocMemoryCounter.Subtract(AllocatedSize);
+#endif
+}
+
+void FFileIoStoreStats::OnBufferReleased()
+{
+#if COUNTERSTRACE_ENABLED
+	AvailableBuffersCounter.Increment();
+#endif
+}
+
+void FFileIoStoreStats::OnBufferAllocated()
+{
+#if COUNTERSTRACE_ENABLED
+	AvailableBuffersCounter.Decrement();
+#endif
+}
+
+#endif
\ No newline at end of file
diff --git a/Engine/Source/Runtime/PakFile/Private/IoDispatcherFileBackend.h b/Engine/Source/Runtime/PakFile/Private/IoDispatcherFileBackend.h
new file mode 100644
--- /dev/null
+++ b/Engine/Source/Runtime/PakFile/Private/IoDispatcherFileBackend.h
@@ -0,0 +1,192 @@
+// Copyright Epic Games, Inc. All Rights Reserved.
+
+#pragma once
+
+#include "IoDispatcherFileBackendTypes.h"
+#include "IO/IoDispatcher.h"
+#include "IO/IoStore.h"
+#include "Containers/Array.h"
+#include "Containers/Map.h"
+#include "Stats/Stats.h"
+#include "Async/TaskGraphInterfaces.h"
+#include "HAL/Runnable.h"
+#include "Misc/AES.h"
+#include "GenericPlatform/GenericPlatformFile.h"
+
+class IMappedFileHandle;
+
+struct FFileIoStoreCompressionContext
+{
+	FFileIoStoreCompressionContext* Next = nullptr;
+	uint64 UncompressedBufferSize = 0;
+	uint8* UncompressedBuffer = nullptr;
+};
+
+class FFileIoStoreReader
+{
+public:
+	FFileIoStoreReader(IPlatformFileIoStore& InPlatformImpl, FFileIoStoreStats& InStats);
+	~FFileIoStoreReader();
+	FIoStatus Initialize(const TCHAR* InTocFilePath, int32 Order);
+	uint32 GetContainerInstanceId() const
+	{
+		return ContainerFile.ContainerInstanceId;
+	}
+	FIoStatus Close();
+	bool DoesChunkExist(const FIoChunkId& ChunkId) const;
+	TIoStatusOr<uint64> GetSizeForChunk(const FIoChunkId& ChunkId) const;
+	const FIoOffsetAndLength* Resolve(const FIoChunkId& ChunkId) const;
+	const FFileIoStoreContainerFile& GetContainerFile() const { return ContainerFile; }
+	IMappedFileHandle* GetMappedContainerFileHandle(uint64 TocOffset);
+	const FIoContainerId& GetContainerId() const { return ContainerId; }
+	int32 GetOrder() const { return Order; }
+	bool IsEncrypted() const { return EnumHasAnyFlags(ContainerFile.ContainerFlags, EIoContainerFlags::Encrypted); }
+	bool IsSigned() const { return EnumHasAnyFlags(ContainerFile.ContainerFlags, EIoContainerFlags::Signed); }
+	const FGuid& GetEncryptionKeyGuid() const { return ContainerFile.EncryptionKeyGuid; }
+	void SetEncryptionKey(const FAES::FAESKey& Key) { ContainerFile.EncryptionKey = Key; }
+	const FAES::FAESKey& GetEncryptionKey() const { return ContainerFile.EncryptionKey; }
+	TIoStatusOr<FIoContainerHeader> ReadContainerHeader() const;
+	void ReopenAllFileHandles();
+
+private:
+	const FIoOffsetAndLength* FindChunkInternal(const FIoChunkId& ChunkId) const;
+	uint64 GetTocAllocatedSize() const;
+	
+	IPlatformFileIoStore& PlatformImpl;
+	FFileIoStoreStats& Stats;
+
+	struct FPerfectHashMap
+	{
+		TArray<int32> TocChunkHashSeeds;
+		TArray<FIoChunkId> TocChunkIds;
+		TArray<FIoOffsetAndLength> TocOffsetAndLengths;
+	};
+	FPerfectHashMap PerfectHashMap;
+	TMap<FIoChunkId, FIoOffsetAndLength> TocImperfectHashMapFallback;
+	FFileIoStoreContainerFile ContainerFile;
+	FIoContainerId ContainerId;
+	int32 Order;
+	bool bClosed = false;
+	bool bHasPerfectHashMap = false;
+
+	static TAtomic<uint32> GlobalPartitionIndex;
+	static TAtomic<uint32> GlobalContainerInstanceId;
+};
+
+class FFileIoStoreRequestTracker
+{
+public:
+	FFileIoStoreRequestTracker(FFileIoStoreRequestAllocator& RequestAllocator, FFileIoStoreRequestQueue& RequestQueue);
+	~FFileIoStoreRequestTracker();
+
+	FFileIoStoreCompressedBlock* FindOrAddCompressedBlock(FFileIoStoreBlockKey Key, bool& bOutWasAdded);
+	void RemoveCompressedBlock(const FFileIoStoreCompressedBlock* CompressedBlock);
+	FFileIoStoreReadRequest* FindOrAddRawBlock(FFileIoStoreBlockKey Key, bool& bOutWasAdded);
+	void RemoveRawBlock(const FFileIoStoreReadRequest* RawBlock);
+	void AddReadRequestsToResolvedRequest(FFileIoStoreCompressedBlock* CompressedBlock, FFileIoStoreResolvedRequest& ResolvedRequest);
+	void AddReadRequestsToResolvedRequest(const FFileIoStoreReadRequestList& Requests, FFileIoStoreResolvedRequest& ResolvedRequest);
+	bool CancelIoRequest(FFileIoStoreResolvedRequest& ResolvedRequest);
+	void UpdatePriorityForIoRequest(FFileIoStoreResolvedRequest& ResolvedRequest);
+	void ReleaseIoRequestReferences(FFileIoStoreResolvedRequest& ResolvedRequest);
+	int64 GetLiveReadRequestsCount() const;
+
+private:
+	FFileIoStoreRequestAllocator& RequestAllocator;
+	FFileIoStoreRequestQueue& RequestQueue;
+	TMap<FFileIoStoreBlockKey, FFileIoStoreCompressedBlock*> CompressedBlocksMap;
+	TMap<FFileIoStoreBlockKey, FFileIoStoreReadRequest*> RawBlocksMap;
+};
+
+class FFileIoStore final
+	: public FRunnable
+	, public IIoDispatcherBackend
+{
+public:
+	FFileIoStore(TUniquePtr<IPlatformFileIoStore>&& PlatformImpl);
+	~FFileIoStore();
+	void Initialize(TSharedRef<const FIoDispatcherBackendContext> Context) override;
+	TIoStatusOr<FIoContainerHeader> Mount(const TCHAR* InTocPath, int32 Order, const FGuid& EncryptionKeyGuid, const FAES::FAESKey& EncryptionKey);
+	bool Unmount(const TCHAR* InTocPath);
+	bool Resolve(FIoRequestImpl* Request) override;
+	void CancelIoRequest(FIoRequestImpl* Request) override;
+	void UpdatePriorityForIoRequest(FIoRequestImpl* Request) override;
+	bool DoesChunkExist(const FIoChunkId& ChunkId) const override;
+	TIoStatusOr<uint64> GetSizeForChunk(const FIoChunkId& ChunkId) const;
+	FIoRequestImpl* GetCompletedRequests() override;
+	TIoStatusOr<FIoMappedRegion> OpenMapped(const FIoChunkId& ChunkId, const FIoReadOptions& Options) override;
+	void ReopenAllFileHandles();
+
+	virtual bool Init() override;
+	virtual uint32 Run() override;
+	virtual void Stop() override;
+
+	uint32 GetThreadId() const;
+
+private:
+	class FDecompressAsyncTask
+	{
+	public:
+		FDecompressAsyncTask(FFileIoStore& InOuter, FFileIoStoreCompressedBlock* InCompressedBlock)
+			: Outer(InOuter)
+			, CompressedBlock(InCompressedBlock)
+		{
+
+		}
+
+		static FORCEINLINE TStatId GetStatId()
+		{
+			RETURN_QUICK_DECLARE_CYCLE_STAT(FIoStoreDecompressTask, STATGROUP_TaskGraphTasks);
+		}
+
+		static ENamedThreads::Type GetDesiredThread();
+
+		FORCEINLINE static ESubsequentsMode::Type GetSubsequentsMode()
+		{
+			return ESubsequentsMode::FireAndForget;
+		}
+
+		void DoTask(ENamedThreads::Type CurrentThread, const FGraphEventRef& MyCompletionGraphEvent)
+		{
+			Outer.ScatterBlock(CompressedBlock, true);
+		}
+
+	private:
+		FFileIoStore& Outer;
+		FFileIoStoreCompressedBlock* CompressedBlock;
+	};
+
+	void OnNewPendingRequestsAdded();
+	void ReadBlocks(FFileIoStoreResolvedRequest& ResolvedRequest);
+	void FreeBuffer(FFileIoStoreBuffer& Buffer);
+	FFileIoStoreCompressionContext* AllocCompressionContext();
+	void FreeCompressionContext(FFileIoStoreCompressionContext* CompressionContext);
+	void ScatterBlock(FFileIoStoreCompressedBlock* CompressedBlock, bool bIsAsync);
+	void CompleteDispatcherRequest(FFileIoStoreResolvedRequest* ResolvedRequest);
+	void FinalizeCompressedBlock(FFileIoStoreCompressedBlock* CompressedBlock);
+	void UpdateAsyncIOMinimumPriority();
+
+	uint64 ReadBufferSize = 0;
+	TSharedPtr<const FIoDispatcherBackendContext> BackendContext;
+	FFileIoStoreStats Stats;
+	FFileIoStoreBlockCache BlockCache;
+	FFileIoStoreBufferAllocator BufferAllocator;
+	FFileIoStoreRequestAllocator RequestAllocator;
+	FFileIoStoreRequestQueue RequestQueue;
+	FFileIoStoreRequestTracker RequestTracker;
+	TUniquePtr<IPlatformFileIoStore> PlatformImpl;
+	FRunnableThread* Thread = nullptr;
+	bool bIsMultithreaded;
+	TAtomic<bool> bStopRequested{ false };
+	mutable FRWLock IoStoreReadersLock;
+	TArray<TUniquePtr<FFileIoStoreReader>> IoStoreReaders;
+	FFileIoStoreCompressionContext* FirstFreeCompressionContext = nullptr;
+	FFileIoStoreCompressedBlock* ReadyForDecompressionHead = nullptr;
+	FFileIoStoreCompressedBlock* ReadyForDecompressionTail = nullptr;
+	FCriticalSection DecompressedBlocksCritical;
+	FFileIoStoreCompressedBlock* FirstDecompressedBlock = nullptr;
+	FIoRequestImpl* CompletedRequestsHead = nullptr;
+	FIoRequestImpl* CompletedRequestsTail = nullptr;
+	EAsyncIOPriorityAndFlags CurrentAsyncIOMinimumPriority = AIOP_MIN;
+};
+
+TSharedRef<FFileIoStore> CreateIoDispatcherFileBackend();
diff --git a/Engine/Source/Runtime/PakFile/Public/IPlatformFilePak.h b/Engine/Source/Runtime/PakFile/Public/IPlatformFilePak.h
--- a/Engine/Source/Runtime/PakFile/Public/IPlatformFilePak.h
+++ b/Engine/Source/Runtime/PakFile/Public/IPlatformFilePak.h
@@ -21,7 +21,7 @@
 
 class FChunkCacheWorker;
 class IAsyncReadFileHandle;
-struct IIoDispatcherFileBackend;
+class FFileIoStore;
 class FFilePackageStore;
 struct FIoContainerHeader;
 
@@ -2033,7 +2033,7 @@ class PAKFILE_API FPakPlatformFile : public IPlatformFile
 	FString IniFileExtension;
 	/** The filename for the gameusersettings ini file, used for excluding ini files, but not gameusersettings */
 	FString GameUserSettingsIniFilename;
-	TSharedPtr<IIoDispatcherFileBackend> IoDispatcherFileBackend;
+	TSharedPtr<FFileIoStore> IoDispatcherFileBackend;
 	TSharedPtr<FFilePackageStore> FilePackageStore;
 
 	FTSTicker::FDelegateHandle RetireReadersHandle;
