Add WindowsPlatformIoDispatcher module with a file backend platform implementation that uses DirectStorage for Windows

#rnx

diff --git a/Engine/Source/Runtime/PakFile/Internal/IoDispatcherFileBackendTypes.h b/Engine/Source/Runtime/PakFile/Internal/IoDispatcherFileBackendTypes.h
--- a/Engine/Source/Runtime/PakFile/Internal/IoDispatcherFileBackendTypes.h
+++ b/Engine/Source/Runtime/PakFile/Internal/IoDispatcherFileBackendTypes.h
@@ -21,8 +21,15 @@
 
 struct FFileIoStoreCompressionContext;
 
-struct FFileIoStoreContainerFilePartition
+struct PAKFILE_API FFileIoStoreContainerFilePartition
 {
+	FFileIoStoreContainerFilePartition() = default;
+	FFileIoStoreContainerFilePartition(FFileIoStoreContainerFilePartition&&) = default;
+	FFileIoStoreContainerFilePartition(const FFileIoStoreContainerFilePartition&) = delete;
+
+	FFileIoStoreContainerFilePartition& operator=(FFileIoStoreContainerFilePartition&&) = default;
+	FFileIoStoreContainerFilePartition& operator=(const FFileIoStoreContainerFilePartition&) = delete;
+
 	uint64 FileHandle = 0;
 	uint64 FileSize = 0;
 	uint32 ContainerFileIndex = 0;
@@ -30,8 +37,15 @@ struct FFileIoStoreContainerFilePartition
 	TUniquePtr<IMappedFileHandle> MappedFileHandle;
 };
 
-struct FFileIoStoreContainerFile
+struct PAKFILE_API FFileIoStoreContainerFile
 {
+	FFileIoStoreContainerFile() = default;
+	FFileIoStoreContainerFile(FFileIoStoreContainerFile&&) = default;
+	FFileIoStoreContainerFile(const FFileIoStoreContainerFile&) = delete;
+
+	FFileIoStoreContainerFile& operator=(FFileIoStoreContainerFile&&) = default;
+	FFileIoStoreContainerFile& operator=(const FFileIoStoreContainerFile&) = delete;
+
 	uint64 PartitionSize = 0;
 	uint64 CompressionBlockSize = 0;
 	TArray<FName> CompressionMethods;
@@ -53,13 +67,13 @@ struct FFileIoStoreContainerFile
 	}
 };
 
-struct FFileIoStoreBuffer
+struct PAKFILE_API FFileIoStoreBuffer
 {
 	FFileIoStoreBuffer* Next = nullptr;
 	uint8* Memory = nullptr;
 };
 
-struct FFileIoStoreBlockKey
+struct PAKFILE_API FFileIoStoreBlockKey
 {
 	union
 	{
@@ -83,7 +97,7 @@ struct FFileIoStoreBlockKey
 	}
 };
 
-struct FFileIoStoreBlockScatter
+struct PAKFILE_API FFileIoStoreBlockScatter
 {
 	struct FFileIoStoreResolvedRequest* Request = nullptr;
 	uint64 DstOffset = 0;
@@ -91,7 +105,7 @@ struct FFileIoStoreBlockScatter
 	uint64 Size = 0;
 };
 
-struct FFileIoStoreCompressedBlock
+struct PAKFILE_API FFileIoStoreCompressedBlock
 {
 	FFileIoStoreCompressedBlock* Next = nullptr;
 	FFileIoStoreBlockKey Key;
@@ -112,7 +126,7 @@ struct FFileIoStoreCompressedBlock
 	bool bCancelled = false;
 };
 
-struct FFileIoStoreReadRequest
+struct PAKFILE_API FFileIoStoreReadRequest
 {
 	enum EQueueStatus
 	{
@@ -236,7 +250,7 @@ private:
 	FFileIoStoreReadRequest* Next = nullptr;;
 };
 
-class FFileIoStoreReadRequestListIterator
+class PAKFILE_API FFileIoStoreReadRequestListIterator
 {
 public:
 	FFileIoStoreReadRequestListIterator(const FFileIoStoreReadRequestListIterator&) = default;
@@ -298,7 +312,7 @@ private:
 
 // Wrapper for doubly-linked intrusive list of FFileIoStoreReadRequest
 
-class FFileIoStoreReadRequestList
+class PAKFILE_API FFileIoStoreReadRequestList
 {
 public:
 	FFileIoStoreReadRequestList()
@@ -522,7 +536,7 @@ private:
 
 class FFileIoStoreStats;
 
-class FFileIoStoreBufferAllocator
+class PAKFILE_API FFileIoStoreBufferAllocator
 {
 public:
 	FFileIoStoreBufferAllocator(FFileIoStoreStats& InStats)
@@ -543,7 +557,7 @@ private:
 	FFileIoStoreBuffer* FirstFreeBuffer = nullptr;
 };
 
-class FFileIoStoreBlockCache
+class PAKFILE_API FFileIoStoreBlockCache
 {
 public:
 	FFileIoStoreBlockCache(FFileIoStoreStats& Stats);
@@ -570,7 +584,7 @@ private:
 	uint64 ReadBufferSize = 0;
 };
 
-struct FFileIoStoreReadRequestSortKey
+struct PAKFILE_API FFileIoStoreReadRequestSortKey
 {
 	uint64 Offset = 0;
 	uint64 Handle = 0;
@@ -584,7 +598,7 @@ struct FFileIoStoreReadRequestSortKey
 };
 
 // Stores FFileIoStoreReadRequest sorted by file handle & offset with a parallel list sorted by insertion order 
-class FFileIoStoreOffsetSortedRequestQueue
+class PAKFILE_API FFileIoStoreOffsetSortedRequestQueue
 {
 public:
 	FFileIoStoreOffsetSortedRequestQueue(int32 InPriority);
@@ -623,7 +637,7 @@ private:
 	static bool RequestSortPredicate(const FFileIoStoreReadRequestSortKey& A, const FFileIoStoreReadRequestSortKey& B);
 };
 
-class FFileIoStoreRequestQueue
+class PAKFILE_API FFileIoStoreRequestQueue
 {
 public:
 	FFileIoStoreReadRequest* Peek();
@@ -736,7 +750,7 @@ private:
 	FSlab* CurrentSlab = nullptr;
 };
 
-struct FFileIoStoreReadRequestLink
+struct PAKFILE_API FFileIoStoreReadRequestLink
 {
 	FFileIoStoreReadRequestLink(FFileIoStoreReadRequest& InReadRequest)
 		: ReadRequest(InReadRequest)
@@ -747,7 +761,7 @@ struct FFileIoStoreReadRequestLink
 	FFileIoStoreReadRequest& ReadRequest;
 };
 
-class FFileIoStoreRequestAllocator
+class PAKFILE_API FFileIoStoreRequestAllocator
 {
 public:
 	int64 GetLiveReadRequestsCount() const
@@ -819,7 +833,7 @@ private:
 	int64 LiveReadRequestsCount = 0;
 };
 
-struct FFileIoStoreResolvedRequest
+struct PAKFILE_API FFileIoStoreResolvedRequest
 {
 public:
 	FFileIoStoreResolvedRequest(
@@ -887,7 +901,7 @@ private:
 
 // Wrapper for sending stats to both insights and csv profiling from any platform-specific dispatcher.
 #if UE_FILEIOSTORE_STATS_ENABLED
-class FFileIoStoreStats
+class PAKFILE_API FFileIoStoreStats
 {
 public:
 	FFileIoStoreStats();
@@ -979,7 +993,7 @@ private:
 	void OnBufferAllocated();
 };
 #else
-class FFileIoStoreStats
+class PAKFILE_API FFileIoStoreStats
 {
 public:
 	void OnFilesystemReadStarted(const FFileIoStoreReadRequest* Request) {}
@@ -1008,7 +1022,7 @@ private:
 };
 #endif
 
-struct FInitializePlatformFileIoStoreParams
+struct PAKFILE_API FInitializePlatformFileIoStoreParams
 {
 	const FWakeUpIoDispatcherThreadDelegate* WakeUpDispatcherThreadDelegate = nullptr;
 	FFileIoStoreRequestAllocator* RequestAllocator = nullptr;
@@ -1041,6 +1055,6 @@ public:
 };
 
 #if PLATFORM_IMPLEMENTS_IO
-CORE_API TUniquePtr<IPlatformFileIoStore> CreatePlatformFileIoStore();
+PAKFILE_API TUniquePtr<IPlatformFileIoStore> CreatePlatformFileIoStore();
 #endif
 
