Add missing include to IoDispatcherFileBackendTypes.h

#rnx

diff --git a/Engine/Source/Runtime/PakFile/Internal/IoDispatcherFileBackendTypes.h b/Engine/Source/Runtime/PakFile/Internal/IoDispatcherFileBackendTypes.h
--- a/Engine/Source/Runtime/PakFile/Internal/IoDispatcherFileBackendTypes.h
+++ b/Engine/Source/Runtime/PakFile/Internal/IoDispatcherFileBackendTypes.h
@@ -4,6 +4,7 @@
 
 #include "IO/IoStore.h"
 #include "IO/IoDispatcherBackend.h"
+#include "Async/MappedFileHandle.h"
 #include "ProfilingDebugging/CsvProfiler.h"
 #include "ProfilingDebugging/CountersTrace.h"
 #include "Containers/Ticker.h"
