IoDispatcher:
Restored container mounting error message log text
More detailed error messages when reading container headers

#rnx

diff --git a/Engine/Source/Runtime/PakFile/Private/IPlatformFilePak.cpp b/Engine/Source/Runtime/PakFile/Private/IPlatformFilePak.cpp
--- a/Engine/Source/Runtime/PakFile/Private/IPlatformFilePak.cpp
+++ b/Engine/Source/Runtime/PakFile/Private/IPlatformFilePak.cpp
@@ -7742,7 +7742,7 @@ bool FPakPlatformFile::Mount(const TCHAR* InPakFilename, uint32 PakOrder, const
 				else
 				{
 					bIoStoreSuccess = false;
-					UE_LOG(LogPakFile, Warning, TEXT("Failed to mount IoStore container \"%s\" [%s]"), *UtocPath);
+					UE_LOG(LogPakFile, Warning, TEXT("Failed to mount IoStore container \"%s\" [%s]"), *UtocPath, *MountResult.Status().ToString());
 				}
 			}
 			else
diff --git a/Engine/Source/Runtime/PakFile/Private/IoDispatcherFileBackend.cpp b/Engine/Source/Runtime/PakFile/Private/IoDispatcherFileBackend.cpp
--- a/Engine/Source/Runtime/PakFile/Private/IoDispatcherFileBackend.cpp
+++ b/Engine/Source/Runtime/PakFile/Private/IoDispatcherFileBackend.cpp
@@ -880,7 +880,7 @@ TIoStatusOr<FIoContainerHeader> FFileIoStoreReader::ReadContainerHeader() const
 	const FIoOffsetAndLength* OffsetAndLength = FindChunkInternal(HeaderChunkId);
 	if (!OffsetAndLength)
 	{
-		return FIoStatus(EIoErrorCode::NotFound);
+		return FIoStatus(FIoStatusBuilder(EIoErrorCode::NotFound) << TEXT("Container header chunk not found"));
 	}
 	
 	const uint64 CompressionBlockSize = ContainerFile.CompressionBlockSize;
@@ -901,17 +901,17 @@ TIoStatusOr<FIoContainerHeader> FFileIoStoreReader::ReadContainerHeader() const
 	TUniquePtr<IFileHandle> ContainerFileHandle(Ipf.OpenRead(*Partition.FilePath));
 	if (!ContainerFileHandle)
 	{
-		return FIoStatus(EIoErrorCode::FileOpenFailed);
+		return FIoStatus(FIoStatusBuilder(EIoErrorCode::FileOpenFailed) << TEXT("Failed to open container file ") << Partition.FilePath);
 	}
 	{
 		TRACE_CPUPROFILER_EVENT_SCOPE(ReadFromContainerFile);
 		if (!ContainerFileHandle->Seek(RawOffset))
 		{
-			return FIoStatus(EIoErrorCode::ReadError);
+			return FIoStatus(FIoStatusBuilder(EIoErrorCode::ReadError) << FString::Printf(TEXT("Failed seeking to offset %llu in container file"), RawOffset));
 		}
 		if (!ContainerFileHandle->Read(IoBuffer.Data(), IoBuffer.DataSize()))
 		{
-			return FIoStatus(EIoErrorCode::ReadError);
+			return FIoStatus(FIoStatusBuilder(EIoErrorCode::ReadError) << FString::Printf(TEXT("Failed reading %llu bytes at offset %llu"), IoBuffer.DataSize(), RawOffset));
 		}
 	}
 
@@ -932,7 +932,7 @@ TIoStatusOr<FIoContainerHeader> FFileIoStoreReader::ReadContainerHeader() const
 				FSHA1::HashBuffer(BlockData, BlockSize, BlockHash.Hash);
 				if (SignatureHash != BlockHash)
 				{
-					return FIoStatus(EIoErrorCode::SignatureError);
+					return FIoStatus(FIoStatusBuilder(EIoErrorCode::SignatureError) << TEXT("Signature error detected when reading container header"));
 				}
 			}
 			if (bEncrypted)
