Adjust LogPakFile  and add error message on unreachable pak mount points
- Added additional check and error message that triggers when a pak file is mounted with a mount point that is not mounted to any root directory
- added VeryVerbose message when access to any file from ExcludedNonPakExtensions is prevented
- Introduced now Display message to print mounted pak file name and mount point, moved the corresponding logs that were already printing that information from the shadercodelibrary to verbose
- Moved most FPakReadRequest messages to VeryVerbose
- Moved PrimaryIndexSize PathHashIndexSize and FullDirectoryIndeSize to Verbose
- Moved OnPakFileMounted2Time to Verbose

diff --git a/Engine/Source/Runtime/PakFile/Private/IPlatformFilePak.cpp b/Engine/Source/Runtime/PakFile/Private/IPlatformFilePak.cpp
--- a/Engine/Source/Runtime/PakFile/Private/IPlatformFilePak.cpp
+++ b/Engine/Source/Runtime/PakFile/Private/IPlatformFilePak.cpp
@@ -35,6 +35,7 @@
 #include "ProfilingDebugging/ScopedTimers.h"
 #include "Async/MappedFileHandle.h"
 #include "IoDispatcherFileBackend.h"
+#include "Misc/PackageName.h"
 
 #include "ProfilingDebugging/LoadTimeTracker.h"
 #include "IO/IoContainerHeader.h"
@@ -1819,7 +1820,7 @@ private: // below here we assume CachedFilesScopeLock until we get to the next s
 
 	void ClearBlock(FCacheBlock &Block)
 	{
-		UE_LOG(LogPakFile, Verbose, TEXT("FPakReadRequest[%016llX, %016llX) ClearBlock"), Block.OffsetAndPakIndex, Block.OffsetAndPakIndex + Block.Size);
+		UE_LOG(LogPakFile, VeryVerbose, TEXT("FPakReadRequest[%016llX, %016llX) ClearBlock"), Block.OffsetAndPakIndex, Block.OffsetAndPakIndex + Block.Size);
 
 		if (Block.Memory)
 		{
@@ -1926,7 +1927,7 @@ private: // below here we assume CachedFilesScopeLock until we get to the next s
 								FCacheBlock &Block = CacheBlockAllocator.Get(BlockIndex);
 								if (!Block.InRequestRefCount)
 								{
-									UE_LOG(LogPakFile, Verbose, TEXT("FPakReadRequest[%016llX, %016llX) Discard Cached"), Block.OffsetAndPakIndex, Block.OffsetAndPakIndex + Block.Size);
+									UE_LOG(LogPakFile, VeryVerbose, TEXT("FPakReadRequest[%016llX, %016llX) Discard Cached"), Block.OffsetAndPakIndex, Block.OffsetAndPakIndex + Block.Size);
 									ClearBlock(Block);
 									return true;
 								}
@@ -1990,7 +1991,7 @@ private: // below here we assume CachedFilesScopeLock until we get to the next s
 								FCacheBlock &Block = CacheBlockAllocator.Get(BlockIndex);
 								if (!Block.InRequestRefCount && (CurrentTime - Block.TimeNoLongerReferenced >= GPakCache_TimeToTrim))
 								{
-									UE_LOG(LogPakFile, Verbose, TEXT("FPakReadRequest[%016llX, %016llX) Discard Cached Based on Time"), Block.OffsetAndPakIndex, Block.OffsetAndPakIndex + Block.Size);
+									UE_LOG(LogPakFile, VeryVerbose, TEXT("FPakReadRequest[%016llX, %016llX) Discard Cached Based on Time"), Block.OffsetAndPakIndex, Block.OffsetAndPakIndex + Block.Size);
 									ClearBlock(Block);
 									return true;
 								}
@@ -2064,7 +2065,7 @@ private: // below here we assume CachedFilesScopeLock until we get to the next s
 							FCacheBlock &Block = CacheBlockAllocator.Get(BlockIndex);
 							if (!Block.InRequestRefCount)
 							{
-								UE_LOG(LogPakFile, Verbose, TEXT("FPakReadRequest[%016llX, %016llX) Discard Cached"), Block.OffsetAndPakIndex, Block.OffsetAndPakIndex + Block.Size);
+								UE_LOG(LogPakFile, VeryVerbose, TEXT("FPakReadRequest[%016llX, %016llX) Discard Cached"), Block.OffsetAndPakIndex, Block.OffsetAndPakIndex + Block.Size);
 								ClearBlock(Block);
 								return true;
 							}
@@ -2180,7 +2181,7 @@ private: // below here we assume CachedFilesScopeLock until we get to the next s
 
 		if (Request.Status == EInRequestStatus::Complete && Request.UniqueID == Request.Owner->UniqueID && RequestIndex == Request.Owner->InRequestIndex &&  Request.OffsetAndPakIndex == Request.Owner->OffsetAndPakIndex)
 		{
-			UE_LOG(LogPakFile, Verbose, TEXT("FPakReadRequest[%016llX, %016llX) Notify complete"), Request.OffsetAndPakIndex, Request.OffsetAndPakIndex + Request.Size);
+			UE_LOG(LogPakFile, VeryVerbose, TEXT("FPakReadRequest[%016llX, %016llX) Notify complete"), Request.OffsetAndPakIndex, Request.OffsetAndPakIndex + Request.Size);
 			Request.Owner->RequestIsComplete();
 			return;
 		}
@@ -2588,7 +2589,7 @@ private: // below here we assume CachedFilesScopeLock until we get to the next s
 		}
 		EAsyncIOPriorityAndFlags Priority = AIOP_Normal; // the lower level requests are not prioritized at the moment
 		check(Block.Status == EBlockStatus::InFlight);
-		UE_LOG(LogPakFile, Verbose, TEXT("FPakReadRequest[%016llX, %016llX) StartBlockTask"), Block.OffsetAndPakIndex, Block.OffsetAndPakIndex + Block.Size);
+		UE_LOG(LogPakFile, VeryVerbose, TEXT("FPakReadRequest[%016llX, %016llX) StartBlockTask"), Block.OffsetAndPakIndex, Block.OffsetAndPakIndex + Block.Size);
 		uint16 PakIndex = GetRequestPakIndex(Block.OffsetAndPakIndex);
 		FPakData& Pak = CachedPakData[PakIndex];
 		RequestsToLower[IndexToFill].BlockIndex = Block.Index;
@@ -2674,7 +2675,7 @@ private: // below here we assume CachedFilesScopeLock until we get to the next s
 			check(Block.Size > 0);
 			DEC_MEMORY_STAT_BY(STAT_AsyncFileMemory, Block.Size);
 			FMemory::Free(Memory);
-			UE_LOG(LogPakFile, Verbose, TEXT("FPakReadRequest[%016llX, %016llX) Cancelled"), Block.OffsetAndPakIndex, Block.OffsetAndPakIndex + Block.Size);
+			UE_LOG(LogPakFile, VeryVerbose, TEXT("FPakReadRequest[%016llX, %016llX) Cancelled"), Block.OffsetAndPakIndex, Block.OffsetAndPakIndex + Block.Size);
 			ClearBlock(Block);
 		}
 		else
@@ -2903,14 +2904,14 @@ public:
 #if USE_PAK_PRECACHE && CSV_PROFILER
 			FPlatformAtomics::InterlockedIncrement(&GPreCacheHotBlocksCount);
 #endif
-			UE_LOG(LogPakFile, Verbose, TEXT("FPakReadRequest[%016llX, %016llX) QueueRequest HOT"), RequestOffsetAndPakIndex, RequestOffsetAndPakIndex + Request.Size);
+			UE_LOG(LogPakFile, VeryVerbose, TEXT("FPakReadRequest[%016llX, %016llX) QueueRequest HOT"), RequestOffsetAndPakIndex, RequestOffsetAndPakIndex + Request.Size);
 		}
 		else
 		{
 #if USE_PAK_PRECACHE && CSV_PROFILER
 			FPlatformAtomics::InterlockedIncrement(&GPreCacheColdBlocksCount);
 #endif
-			UE_LOG(LogPakFile, Verbose, TEXT("FPakReadRequest[%016llX, %016llX) QueueRequest COLD"), RequestOffsetAndPakIndex, RequestOffsetAndPakIndex + Request.Size);
+			UE_LOG(LogPakFile, VeryVerbose, TEXT("FPakReadRequest[%016llX, %016llX) QueueRequest COLD"), RequestOffsetAndPakIndex, RequestOffsetAndPakIndex + Request.Size);
 		}
 
 		TrimCache();
@@ -3938,7 +3939,7 @@ public:
 			Blocks.AddDefaulted(FileEntry.CompressionBlocks.Num());
 			CompressedChunkOffset = InPakFile->GetInfo().HasRelativeCompressedChunkOffsets() ? FileEntry.Offset : 0;
 		}
-		UE_LOG(LogPakFile, Verbose, TEXT("FPakPlatformFile::OpenAsyncRead[%016llX, %016llX) %s"), OffsetInPak, OffsetInPak + CompressedFileSize, Filename);
+		UE_LOG(LogPakFile, VeryVerbose, TEXT("FPakPlatformFile::OpenAsyncRead[%016llX, %016llX) %s"), OffsetInPak, OffsetInPak + CompressedFileSize, Filename);
 		check(PakFileSize > 0 && OffsetInPak + CompressedFileSize <= PakFileSize && OffsetInPak >= 0);
 
 		ReadCallbackFunction = [this](bool bWasCancelled, IAsyncReadRequest* Request)
@@ -4516,7 +4517,7 @@ public:
 		UncompressedFileSize = FileEntry.UncompressedSize;
 		int64 CompressedFileSize = FileEntry.UncompressedSize;
 		check(FileEntry.CompressionMethodIndex == 0);
-		UE_LOG(LogPakFile, Verbose, TEXT("FPakPlatformFile::OpenAsyncRead (FBypassPakAsyncReadFileHandle)[%016llX, %016llX) %s"), OffsetInPak, OffsetInPak + CompressedFileSize, Filename);
+		UE_LOG(LogPakFile, VeryVerbose, TEXT("FPakPlatformFile::OpenAsyncRead (FBypassPakAsyncReadFileHandle)[%016llX, %016llX) %s"), OffsetInPak, OffsetInPak + CompressedFileSize, Filename);
 		check(PakFileSize > 0 && OffsetInPak + CompressedFileSize <= PakFileSize && OffsetInPak >= 0);
 
 		LowerHandle = IPlatformFile::GetPlatformPhysical().OpenAsyncRead(*InPakFile->GetFilename());
@@ -5131,6 +5132,7 @@ bool FPakPlatformFile::IsNonPakFilenameAllowed(const FString& InFilename)
 	{
 		FName Ext = FName(*FPaths::GetExtension(InFilename));
 		bAllowed = !ExcludedNonPakExtensions.Contains(Ext);
+		UE_CLOG(!bAllowed, LogPakFile, VeryVerbose, TEXT("Access to file '%s' is limited to pak contents due to file extension being listed in ExcludedNonPakExtensions."), *InFilename)
 	}
 #endif
 
@@ -5738,9 +5740,9 @@ bool FPakFile::LoadIndexInternal(FArchive& Reader)
 #endif
 	}
 
-	UE_LOG(LogPakFile, Log, TEXT("PakFile PrimaryIndexSize=%d"), Info.IndexSize);
-	UE_LOG(LogPakFile, Log, TEXT("PakFile PathHashIndexSize=%d"), PathHashIndexSize);
-	UE_LOG(LogPakFile, Log, TEXT("PakFile FullDirectoryIndexSize=%d"), FullDirectoryIndexSize);
+	UE_LOG(LogPakFile, Verbose, TEXT("PakFile PrimaryIndexSize=%d"), Info.IndexSize);
+	UE_LOG(LogPakFile, Verbose, TEXT("PakFile PathHashIndexSize=%d"), PathHashIndexSize);
+	UE_LOG(LogPakFile, Verbose, TEXT("PakFile FullDirectoryIndexSize=%d"), FullDirectoryIndexSize);
 
 	check(bHasFullDirectoryIndex || bHasPathHashIndex);
 	return true;
@@ -7763,7 +7765,15 @@ bool FPakPlatformFile::Mount(const TCHAR* InPakFilename, uint32 PakOrder, const
 				FScopedDurationTimer Timer(OnPakFileMounted2Time);
 				FCoreDelegates::OnPakFileMounted2.Broadcast(*Pak);
 			}
-			UE_LOG(LogPakFile, Log, TEXT("OnPakFileMounted2Time == %lf"), OnPakFileMounted2Time);
+
+			UE_LOG(LogPakFile, Display, TEXT("Mounted Pak file '%s', mount point: '%s'"), InPakFilename, *Pak->GetMountPoint());
+			UE_LOG(LogPakFile, Verbose, TEXT("OnPakFileMounted2Time == %lf"), OnPakFileMounted2Time);
+			
+			FString OutPackageName;
+			if (!FPackageName::TryConvertFilenameToLongPackageName(Pak->GetMountPoint(), OutPackageName))
+			{
+				UE_LOG(LogPakFile, Display, TEXT("Mount point: '%s' is not mounted to a valid Root Path yet, assets in this pak file may not be accessible until a corresponding UFS Mount Point is added through FPackageName::RegisterMountPoint."), *Pak->GetMountPoint());
+			}
 		}
 		else
 		{
