Fix FFileIoStoreContainerFile::FilePath not being set in FFileIoStoreReader::Initialize.  This made it impossible to unmount iostore containers which kept them from being deleted.
#test Used VK Edit iteration

diff --git a/Engine/Source/Runtime/PakFile/Private/IoDispatcherFileBackend.cpp b/Engine/Source/Runtime/PakFile/Private/IoDispatcherFileBackend.cpp
--- a/Engine/Source/Runtime/PakFile/Private/IoDispatcherFileBackend.cpp
+++ b/Engine/Source/Runtime/PakFile/Private/IoDispatcherFileBackend.cpp
@@ -687,6 +687,7 @@ FIoStatus FFileIoStoreReader::Initialize(const TCHAR* InTocFilePath, int32 InOrd
 		return FIoStatusBuilder(EIoErrorCode::FileOpenFailed) << TEXT("Expected .utoc extension on container path '") << InTocFilePath << TEXT("'");
 	}
 	FStringView BasePathView = ContainerPathView.LeftChop(5);
+	ContainerFile.FilePath = BasePathView;
 
 	IPlatformFile& Ipf = FPlatformFileManager::Get().GetPlatformFile();
 
