Removed option to run with package data in iostore and shaders in pak
Now that all asset data will be in either iostore or pak setting async io priority based on the iodispatcher request queue contents is no longer needed
#rnx

diff --git a/Engine/Config/BaseEngine.ini b/Engine/Config/BaseEngine.ini
--- a/Engine/Config/BaseEngine.ini
+++ b/Engine/Config/BaseEngine.ini
@@ -1112,7 +1112,6 @@ UseSeperateBulkDataFiles=False
 GameThreadHeartBeatStartSuspended=False
 GameThreadHeartBeatHitchDuration=-1.0
 GameThreadHeartBeatStackWalk=false
-AllowShadersInIoStore=True
 IgnoreInlineBulkDataReloadEnsures=False
 UsePackageTrailer=False
 HangDuration=0.0
diff --git a/Engine/Source/Programs/AutomationTool/Scripts/CopyBuildToStagingDirectory.Automation.cs b/Engine/Source/Programs/AutomationTool/Scripts/CopyBuildToStagingDirectory.Automation.cs
--- a/Engine/Source/Programs/AutomationTool/Scripts/CopyBuildToStagingDirectory.Automation.cs
+++ b/Engine/Source/Programs/AutomationTool/Scripts/CopyBuildToStagingDirectory.Automation.cs
@@ -2877,13 +2877,6 @@ namespace AutomationScripts
 						Dictionary<string, string> UnrealPakResponseFile = PakParams.UnrealPakResponseFile;
 						if (bShouldCreateIoStoreContainerFiles)
 						{
-							bool bAllowShadersInIoStore;
-							if (!PlatformEngineConfig.GetBool("Core.System", "AllowShadersInIoStore", out bAllowShadersInIoStore))
-							{
-								bAllowShadersInIoStore = false;
-							}
-
-
 							UnrealPakResponseFile = new Dictionary<string, string>();
 							Dictionary<string, string> IoStoreResponseFile = new Dictionary<string, string>();
 							foreach (var Entry in PakParams.UnrealPakResponseFile)
@@ -2907,14 +2900,7 @@ namespace AutomationScripts
 								}
 								else if (Path.GetExtension(Entry.Key).Contains(".ushaderbytecode"))
 								{
-									if (bAllowShadersInIoStore)
-									{
-										IoStoreResponseFile.Add(Entry.Key, Entry.Value);
-									}
-									else
-									{
-										UnrealPakResponseFile.Add(Entry.Key, Entry.Value);
-									}
+									IoStoreResponseFile.Add(Entry.Key, Entry.Value);
 								}
 								else if (!Path.GetExtension(Entry.Key).Contains(".uexp"))
 								{
diff --git a/Engine/Source/Runtime/PakFile/Internal/IoDispatcherFileBackendTypes.h b/Engine/Source/Runtime/PakFile/Internal/IoDispatcherFileBackendTypes.h
--- a/Engine/Source/Runtime/PakFile/Internal/IoDispatcherFileBackendTypes.h
+++ b/Engine/Source/Runtime/PakFile/Internal/IoDispatcherFileBackendTypes.h
@@ -616,7 +616,6 @@ public:
 	// Remove all requests whose priority has been changed to something other than the Priority of this queue
 	TArray<FFileIoStoreReadRequest*> RemoveMisprioritizedRequests();
 
-	FFileIoStoreReadRequest* Peek(FFileIoStoreReadRequestSortKey LastSortKey);
 	FFileIoStoreReadRequest* Pop(FFileIoStoreReadRequestSortKey LastSortKey);
 	void Push(FFileIoStoreReadRequest* Request);
 	void CancelRequestsWithFileHandle(uint64 FileHandle);
@@ -641,7 +640,6 @@ private:
 class PAKFILE_API FFileIoStoreRequestQueue
 {
 public:
-	FFileIoStoreReadRequest* Peek();
 	FFileIoStoreReadRequest* Pop();
 	void Push(FFileIoStoreReadRequest& Request);	// Takes ownership of Request and rewrites its intrustive linked list pointers
 	void Push(FFileIoStoreReadRequestList& Requests); // Consumes the request list and overwrites all intrustive linked list pointers
diff --git a/Engine/Source/Runtime/PakFile/Private/IoDispatcherFileBackend.cpp b/Engine/Source/Runtime/PakFile/Private/IoDispatcherFileBackend.cpp
--- a/Engine/Source/Runtime/PakFile/Private/IoDispatcherFileBackend.cpp
+++ b/Engine/Source/Runtime/PakFile/Private/IoDispatcherFileBackend.cpp
@@ -400,11 +400,6 @@ FFileIoStoreReadRequest* FFileIoStoreOffsetSortedRequestQueue::GetNextInternal(F
 	return Request;
 }
 
-FFileIoStoreReadRequest* FFileIoStoreOffsetSortedRequestQueue::Peek(FFileIoStoreReadRequestSortKey LastSortKey)
-{
-	return GetNextInternal(LastSortKey, false);
-}
-
 FFileIoStoreReadRequest* FFileIoStoreOffsetSortedRequestQueue::Pop(FFileIoStoreReadRequestSortKey LastSortKey)
 {
 	return GetNextInternal(LastSortKey, true);
@@ -469,35 +464,6 @@ void FFileIoStoreRequestQueue::UpdateSortRequestsByOffset()
 	}
 }
 
-FFileIoStoreReadRequest* FFileIoStoreRequestQueue::Peek()
-{
-	TRACE_CPUPROFILER_EVENT_SCOPE(RequestQueuePeek);
-	FScopeLock _(&CriticalSection);
-	UpdateSortRequestsByOffset();
-	if (bSortRequestsByOffset)
-	{
-		if( SortedPriorityQueues.Num() == 0)
-		{
-			return nullptr;
-		}
-
-		FFileIoStoreOffsetSortedRequestQueue& SubQueue = SortedPriorityQueues.Last();
-		check(!SubQueue.IsEmpty());
-		FFileIoStoreReadRequest* Request = SubQueue.Peek(LastSortKey);
-		check(Request);
-		// Do not update LastUsedPriority here until we actually pop
-		return Request;
-	}
-	else
-	{
-		if (Heap.Num() == 0)
-		{
-			return nullptr;
-		}
-		return Heap.HeapTop();
-	}
-}
-
 FFileIoStoreReadRequest* FFileIoStoreRequestQueue::Pop()
 {
 	TRACE_CPUPROFILER_EVENT_SCOPE(RequestQueuePop);
@@ -1992,28 +1958,6 @@ void FFileIoStore::FreeCompressionContext(FFileIoStoreCompressionContext* Compre
 	FirstFreeCompressionContext = CompressionContext;
 }
 
-void FFileIoStore::UpdateAsyncIOMinimumPriority()
-{
-	EAsyncIOPriorityAndFlags NewAsyncIOMinimumPriority = AIOP_MIN;
-	if (FFileIoStoreReadRequest* NextRequest = RequestQueue.Peek())
-	{
-		if (NextRequest->Priority >= IoDispatcherPriority_High)
-		{
-			NewAsyncIOMinimumPriority = AIOP_MAX;
-		}
-		else if (NextRequest->Priority >= IoDispatcherPriority_Medium)
-		{
-			NewAsyncIOMinimumPriority = AIOP_Normal;
-		}
-	}
-	if (NewAsyncIOMinimumPriority != CurrentAsyncIOMinimumPriority)
-	{
-		//TRACE_BOOKMARK(TEXT("SetAsyncMinimumPrioirity(%d)"), NewAsyncIOMinimumPriority);
-		FPlatformFileManager::Get().GetPlatformFile().SetAsyncMinimumPriority(NewAsyncIOMinimumPriority);
-		CurrentAsyncIOMinimumPriority = NewAsyncIOMinimumPriority;
-	}
-}
-
 bool FFileIoStore::Init()
 {
 	return true;
@@ -2029,10 +1973,8 @@ uint32 FFileIoStore::Run()
 {
 	while (!bStopRequested)
 	{
-		UpdateAsyncIOMinimumPriority();
 		if (!PlatformImpl->StartRequests(RequestQueue))
 		{
-			UpdateAsyncIOMinimumPriority();
 			PlatformImpl->ServiceWait();
 		}
 	}
diff --git a/Engine/Source/Runtime/PakFile/Private/IoDispatcherFileBackend.h b/Engine/Source/Runtime/PakFile/Private/IoDispatcherFileBackend.h
--- a/Engine/Source/Runtime/PakFile/Private/IoDispatcherFileBackend.h
+++ b/Engine/Source/Runtime/PakFile/Private/IoDispatcherFileBackend.h
@@ -163,7 +163,6 @@ private:
 	void ScatterBlock(FFileIoStoreCompressedBlock* CompressedBlock, bool bIsAsync);
 	void CompleteDispatcherRequest(FFileIoStoreResolvedRequest* ResolvedRequest);
 	void FinalizeCompressedBlock(FFileIoStoreCompressedBlock* CompressedBlock);
-	void UpdateAsyncIOMinimumPriority();
 
 	uint64 ReadBufferSize = 0;
 	TSharedPtr<const FIoDispatcherBackendContext> BackendContext;
@@ -186,7 +185,6 @@ private:
 	FFileIoStoreCompressedBlock* FirstDecompressedBlock = nullptr;
 	FIoRequestImpl* CompletedRequestsHead = nullptr;
 	FIoRequestImpl* CompletedRequestsTail = nullptr;
-	EAsyncIOPriorityAndFlags CurrentAsyncIOMinimumPriority = AIOP_MIN;
 };
 
 TSharedRef<FFileIoStore> CreateIoDispatcherFileBackend();
