IPlatformFilePak:
Added additional check and error message that triggers when a pak file is mounted with a mount point that is not mounted to any root directory

diff --git a/Engine/Source/Runtime/PakFile/Private/IPlatformFilePak.cpp b/Engine/Source/Runtime/PakFile/Private/IPlatformFilePak.cpp
--- a/Engine/Source/Runtime/PakFile/Private/IPlatformFilePak.cpp
+++ b/Engine/Source/Runtime/PakFile/Private/IPlatformFilePak.cpp
@@ -35,6 +35,7 @@
 #include "ProfilingDebugging/ScopedTimers.h"
 #include "Async/MappedFileHandle.h"
 #include "IoDispatcherFileBackend.h"
+#include "Misc/PackageName.h"
 
 #include "ProfilingDebugging/LoadTimeTracker.h"
 #include "IO/IoContainerHeader.h"
@@ -7767,6 +7768,16 @@ bool FPakPlatformFile::Mount(const TCHAR* InPakFilename, uint32 PakOrder, const
 
 			UE_LOG(LogPakFile, Display, TEXT("Mounted Pak file '%s', mount point: '%s'"), InPakFilename, *Pak->GetMountPoint());
 			UE_LOG(LogPakFile, Verbose, TEXT("OnPakFileMounted2Time == %lf"), OnPakFileMounted2Time);
+						
+			// skip this check for the default mountpoint, it will print false positives
+			if (FPaths::CreateStandardFilename(Pak->GetMountPoint()) != FPaths::CreateStandardFilename(FPaths::RootDir()))
+			{
+				FString OutPackageName;
+				if (!FPackageName::TryConvertFilenameToLongPackageName(Pak->GetMountPoint(), OutPackageName))
+				{
+					UE_LOG(LogPakFile, Display, TEXT("Mount point: '%s' is not mounted to a valid Root Path yet, assets in this pak file may not be accessible until a corresponding UFS Mount Point is added through FPackageName::RegisterMountPoint."), *Pak->GetMountPoint());
+				}
+			}
 		}
 		else
 		{
