- Added a -LookLooseFirst commandline option that tells pak file system to prefer loose files before looking in the pak, instead of falling back to loose when not found in the pak (this is useufl for quickly patching a pak with some local files for debugging)
#jria none

diff --git a/Engine/Source/Runtime/PakFile/Private/IPlatformFilePak.cpp b/Engine/Source/Runtime/PakFile/Private/IPlatformFilePak.cpp
--- a/Engine/Source/Runtime/PakFile/Private/IPlatformFilePak.cpp
+++ b/Engine/Source/Runtime/PakFile/Private/IPlatformFilePak.cpp
@@ -4744,6 +4744,17 @@ IMappedFileHandle* FPakPlatformFile::OpenMapped(const TCHAR* Filename)
 		return nullptr;
 	}
 
+#if !UE_BUILD_SHIPPING
+	if (bLookLooseFirst && IsNonPakFilenameAllowed(Filename))
+	{
+		IMappedFileHandle* Handle = LowerLevel->OpenMapped(Filename);
+		if (Handle != nullptr)
+		{
+			return Handle;
+		}
+	}
+#endif
+
 	// Check pak files first
 	FPakEntry FileEntry;
 	TRefCountPtr<FPakFile> PakEntry;
@@ -7415,6 +7426,9 @@ bool FPakPlatformFile::Initialize(IPlatformFile* Inner, const TCHAR* CmdLine)
 	FString StartupPaksWildcard = GMountStartupPaksWildCard;
 #if !UE_BUILD_SHIPPING
 	FParse::Value(FCommandLine::Get(), TEXT("StartupPaksWildcard="), StartupPaksWildcard);
+
+	// initialize the bLookLooseFirst setting
+	bLookLooseFirst = FParse::Param(FCommandLine::Get(), TEXT("LookLooseFirst"));
 #endif
 
 	FString GlobalUTocPath = FString::Printf(TEXT("%sPaks/global.utoc"), *FPaths::ProjectContentDir());
@@ -8102,6 +8116,19 @@ void FPakPlatformFile::RegisterEncryptionKey(const FGuid& InGuid, const FAES::FA
 IFileHandle* FPakPlatformFile::OpenRead(const TCHAR* Filename, bool bAllowWrite)
 {
 	IFileHandle* Result = NULL;
+
+#if !UE_BUILD_SHIPPING
+	if (bLookLooseFirst && IsNonPakFilenameAllowed(Filename))
+	{
+		Result = LowerLevel->OpenRead(Filename, bAllowWrite);
+		if (Result != nullptr)
+		{
+			return Result;
+		}
+	}
+#endif
+
+
 	TRefCountPtr<FPakFile> PakFile;
 	FPakEntry FileEntry;
 	if (FindFileInPakFiles(Filename, &PakFile, &FileEntry))
@@ -8207,6 +8234,14 @@ bool FPakPlatformFile::BufferedCopyFile(IFileHandle& Dest, IFileHandle& Source,
 
 bool FPakPlatformFile::CopyFile(const TCHAR* To, const TCHAR* From, EPlatformFileRead ReadFlags, EPlatformFileWrite WriteFlags)
 {
+#if !UE_BUILD_SHIPPING
+	if (bLookLooseFirst && LowerLevel->FileExists(From))
+	{
+		return LowerLevel->CopyFile(To, From, ReadFlags, WriteFlags);
+	}
+#endif
+
+
 	bool Result = false;
 	FPakEntry FileEntry;
 	TRefCountPtr<FPakFile> PakFile;
diff --git a/Engine/Source/Runtime/PakFile/Public/IPlatformFilePak.h b/Engine/Source/Runtime/PakFile/Public/IPlatformFilePak.h
--- a/Engine/Source/Runtime/PakFile/Public/IPlatformFilePak.h
+++ b/Engine/Source/Runtime/PakFile/Public/IPlatformFilePak.h
@@ -2038,6 +2038,12 @@ class PAKFILE_API FPakPlatformFile : public IPlatformFile
 
 	FTSTicker::FDelegateHandle RetireReadersHandle;
 
+#if !UE_BUILD_SHIPPING
+	// if true (via -looklocalfirst) then loose/non-ufe files will be looked for before looking in the .pak file
+	// this respects IsNonPakFilenameAllowed()
+	bool bLookLooseFirst = false;
+#endif
+
 	/**
 	 * Gets mounted pak files
 	 */
