- Fixed comment typo

diff --git a/Engine/Source/Runtime/PakFile/Public/IPlatformFilePak.h b/Engine/Source/Runtime/PakFile/Public/IPlatformFilePak.h
--- a/Engine/Source/Runtime/PakFile/Public/IPlatformFilePak.h
+++ b/Engine/Source/Runtime/PakFile/Public/IPlatformFilePak.h
@@ -2039,7 +2039,7 @@ class PAKFILE_API FPakPlatformFile : public IPlatformFile
 	FTSTicker::FDelegateHandle RetireReadersHandle;
 
 #if !UE_BUILD_SHIPPING
-	// if true (via -looklocalfirst) then loose/non-ufe files will be looked for before looking in the .pak file
+	// if true (via -looklocalfirst) then loose/non-ufs files will be looked for before looking in the .pak file
 	// this respects IsNonPakFilenameAllowed()
 	bool bLookLooseFirst = false;
 #endif
