Reduce string allocations/copies in FindFilesAtPathInIndex

diff --git a/Engine/Source/Runtime/PakFile/Public/IPlatformFilePak.h b/Engine/Source/Runtime/PakFile/Public/IPlatformFilePak.h
--- a/Engine/Source/Runtime/PakFile/Public/IPlatformFilePak.h
+++ b/Engine/Source/Runtime/PakFile/Public/IPlatformFilePak.h
@@ -1695,7 +1695,7 @@ private:
 		TArray<FString> DirectoriesInPak; // List of all unique directories at path
 		for (TMap<FString, FPakDirectory>::TConstIterator It(TargetIndex); It; ++It)
 		{
-			const FString PakPath(PakPathCombine(MountPoint, It.Key()));
+			FString PakPath(PakPathCombine(MountPoint, It.Key()));
 			// Check if the file is under the specified path.
 			if (PakPath.StartsWith(Directory))
 			{
@@ -1713,7 +1713,7 @@ private:
 					{
 						if (Directory != PakPath)
 						{
-							DirectoriesInPak.Add(PakPath);
+							DirectoriesInPak.Add(MoveTemp(PakPath));
 						}
 					}
 				}
@@ -1736,7 +1736,7 @@ private:
 				}
 			}
 		}
-		OutFiles.Append(DirectoriesInPak);
+		OutFiles.Append(MoveTemp(DirectoriesInPak));
 	}
 
 	/** Converts the path to a RelativePathFromMount and normalizes it to the expected format for Pak Directories.  Returns false if Path is not under the MountDir and hence can not be in this PakFile. */
