Support -UsePaks on monolithic editor with unccoked data

diff --git a/Engine/Source/Runtime/PakFile/Private/IPlatformFilePak.cpp b/Engine/Source/Runtime/PakFile/Private/IPlatformFilePak.cpp
--- a/Engine/Source/Runtime/PakFile/Private/IPlatformFilePak.cpp
+++ b/Engine/Source/Runtime/PakFile/Private/IPlatformFilePak.cpp
@@ -7372,15 +7372,7 @@ bool FPakPlatformFile::CheckIfPakFilesExist(IPlatformFile* LowLevelFile, const T
 
 bool FPakPlatformFile::ShouldBeUsed(IPlatformFile* Inner, const TCHAR* CmdLine) const
 {
-	bool Result = false;
-#if (!WITH_EDITOR || IS_MONOLITHIC)
-	if (!FParse::Param(CmdLine, TEXT("NoPak")))
-	{
-		TArray<FString> PakFolders;
-		GetPakFolders(CmdLine, PakFolders);
-		Result = CheckIfPakFilesExist(Inner, PakFolders);
-	}
-#else
+#if WITH_EDITOR
 	if (FParse::Param(CmdLine, TEXT("UsePaks")))
 	{
 		TArray<FString> PakFolders;
@@ -7389,9 +7381,19 @@ bool FPakPlatformFile::ShouldBeUsed(IPlatformFile* Inner, const TCHAR* CmdLine)
 		{
 			UE_LOG(LogPakFile, Warning, TEXT("No Pak files were found when checking to make Pak Environment"));
 		}
-		Result = true;
+		return true;
 	}
-#endif
+#endif //if WITH_EDITOR
+
+	bool Result = false;
+#if (!WITH_EDITOR || IS_MONOLITHIC)
+	if (!FParse::Param(CmdLine, TEXT("NoPak")))
+	{
+		TArray<FString> PakFolders;
+		GetPakFolders(CmdLine, PakFolders);
+		Result = CheckIfPakFilesExist(Inner, PakFolders);
+	}
+#endif //if (!WITH_EDITOR || IS_MONOLITHIC)
 	return Result;
 }
 
