Update FFilePackageStore::UncookedPackagesMap when unmounting content

diff --git a/Engine/Source/Runtime/PakFile/Private/FilePackageStore.cpp b/Engine/Source/Runtime/PakFile/Private/FilePackageStore.cpp
--- a/Engine/Source/Runtime/PakFile/Private/FilePackageStore.cpp
+++ b/Engine/Source/Runtime/PakFile/Private/FilePackageStore.cpp
@@ -13,6 +13,20 @@
 
 DEFINE_LOG_CATEGORY_STATIC(LogFilePackageStore, Log, All);
 
+namespace FilePackageStore
+{
+	FStringView GetRootPathFromPackageName(const FStringView Path)
+	{
+		FStringView RootPath;
+		int32 SecondForwardSlash = INDEX_NONE;
+		if (ensure(FStringView(Path.GetData() + 1, Path.Len() - 1).FindChar(TEXT('/'), SecondForwardSlash)))
+		{
+			RootPath = FStringView(Path.GetData(), SecondForwardSlash + 2);
+		}
+		return RootPath;
+	}
+}
+
 thread_local int32 FFilePackageStore::LockedOnThreadCount = 0;
 
 FFilePackageStore::FFilePackageStore()
@@ -22,20 +36,34 @@ FFilePackageStore::FFilePackageStore()
 		{
 			{
 				FScopeLock _(&UncookedPackageRootsLock);
-				PendingUncookedPackageRoots.Add(InFilesystemPath);
+				PendingAddUncookedPackageRoots.Add(InFilesystemPath);
 			}
 			{
 				FWriteScopeLock _(EntriesLock);
 				bNeedsUpdate = true;
 			}
 		});
-#endif
+
+	OnContentPathDismountedDelegateHandle = FPackageName::OnContentPathDismounted().AddLambda([this](const FString& InAssetPath, const FString& InFilesystemPath)
+		{
+			{
+				FScopeLock _(&UncookedPackageRootsLock);
+				PendingRemoveUncookedPackageRoots.Add(InAssetPath);
+				PendingAddUncookedPackageRoots.Remove(InFilesystemPath);
+			}
+			{
+				FWriteScopeLock _(EntriesLock);
+				bNeedsUpdate = true;
+			}
+		});
+#endif //if WITH_EDITOR
 }
 
 FFilePackageStore::~FFilePackageStore()
 {
 #if WITH_EDITOR
 	FPackageName::OnContentPathMounted().Remove(OnContentPathMountedDelegateHandle);
+	FPackageName::OnContentPathDismounted().Remove(OnContentPathDismountedDelegateHandle);
 #endif
 }
 
@@ -46,13 +74,13 @@ void FFilePackageStore::Initialize()
 	FPackageName::QueryRootContentPaths(RootPaths);
 	{
 		FScopeLock _(&UncookedPackageRootsLock);
-		PendingUncookedPackageRoots.Append(RootPaths);
+		PendingAddUncookedPackageRoots.Append(RootPaths);
 	}
 	{
 		FWriteScopeLock _(EntriesLock);
 		bNeedsUpdate = true;
 	}
-#endif
+#endif //if WITH_EDITOR
 }
 
 void FFilePackageStore::Lock()
@@ -225,7 +253,7 @@ void FFilePackageStore::Update()
 			OptionalSegmentStoreEntriesMap.FindOrAdd(PackageId, &OptionalSegmentStoreEntry);
 			++Index;
 		}
-#endif
+#endif //if WITH_EDITOR
 
 		for (const FIoContainerHeaderLocalizedPackage& LocalizedPackage : ContainerHeader->LocalizedPackages)
 		{
@@ -248,19 +276,28 @@ void FFilePackageStore::Update()
 #if WITH_EDITOR
 	{
 		FScopeLock _(&UncookedPackageRootsLock);
-		uint64 TotalAddedCount = 0;
-		if (!PendingUncookedPackageRoots.IsEmpty())
+
+		if (!PendingRemoveUncookedPackageRoots.IsEmpty())
 		{
-			UE_LOG(LogFilePackageStore, Display, TEXT("Searching for uncooked packages in %d new roots..."), PendingUncookedPackageRoots.Num());
-			for (const FString& RootPath : PendingUncookedPackageRoots)
+			UE_LOG(LogFilePackageStore, Display, TEXT("Removing uncooked packages from %d old roots..."), PendingRemoveUncookedPackageRoots.Num());
+			uint64 TotalRemovedCount = RemoveUncookedPackagesFromRoot(PendingRemoveUncookedPackageRoots);
+			PendingRemoveUncookedPackageRoots.Empty();
+			UE_LOG(LogFilePackageStore, Display, TEXT("Removed %lld uncooked packages"), TotalRemovedCount);
+		}
+
+		if (!PendingAddUncookedPackageRoots.IsEmpty())
+		{
+			UE_LOG(LogFilePackageStore, Display, TEXT("Adding uncooked packages from %d new roots..."), PendingAddUncookedPackageRoots.Num());
+			uint64 TotalAddedCount = 0;
+			for (const FString& RootPath : PendingAddUncookedPackageRoots)
 			{
 				TotalAddedCount += AddUncookedPackagesFromRoot(RootPath);
 			}
-			PendingUncookedPackageRoots.Empty();
-			UE_LOG(LogFilePackageStore, Display, TEXT("Found %lld uncooked packages"), TotalAddedCount);
+			PendingAddUncookedPackageRoots.Empty();
+			UE_LOG(LogFilePackageStore, Display, TEXT("Added %lld uncooked packages"), TotalAddedCount);
 		}
 	}
-#endif
+#endif //if WITH_EDITOR
 
 	bNeedsUpdate = false;
 }
@@ -268,8 +305,8 @@ void FFilePackageStore::Update()
 #if WITH_EDITOR
 uint64 FFilePackageStore::AddUncookedPackagesFromRoot(const FString& RootPath)
 {
-	uint64 TotalAddedCount = 0;
-	FPackageName::IteratePackagesInDirectory(RootPath, [this, &RootPath, &TotalAddedCount](const TCHAR* InPackageFileName) -> bool
+	uint64 Count = 0;
+	FPackageName::IteratePackagesInDirectory(RootPath, [this, &RootPath, &Count](const TCHAR* InPackageFileName) -> bool
 		{
 			FPackagePath PackagePath = FPackagePath::FromLocalPath(InPackageFileName);
 			FName PackageName = PackagePath.GetPackageFName();
@@ -279,12 +316,35 @@ uint64 FFilePackageStore::AddUncookedPackagesFromRoot(const FString& RootPath)
 				FUncookedPackage& UncookedPackage = UncookedPackagesMap.FindOrAdd(PackageId);
 				UncookedPackage.PackageName = PackageName;
 				UncookedPackage.HeaderExtension = PackagePath.GetHeaderExtension();
-				++TotalAddedCount;
+				++Count;
 			}
 			return true;
 		});
-	return TotalAddedCount;
+	return Count;
 }
-#endif
+
+uint64 FFilePackageStore::RemoveUncookedPackagesFromRoot(const TSet<FString>& RootPaths)
+{
+	uint64 Count = 0;
+	TMap<FPackageId, FUncookedPackage> RemainingUncookedPackagesMap;
+	RemainingUncookedPackagesMap.Reserve(UncookedPackagesMap.Num());
+	for (auto& Pair : UncookedPackagesMap)
+	{
+		const FNameBuilder PackageName(Pair.Value.PackageName);
+		const FStringView PackageRootPath = FilePackageStore::GetRootPathFromPackageName(PackageName);
+		if (RootPaths.ContainsByHash(GetTypeHash(PackageRootPath), PackageRootPath))
+		{
+			++Count;
+		}
+		else
+		{
+			RemainingUncookedPackagesMap.Add(Pair.Key, MoveTemp(Pair.Value));
+		}
+	}
+	UncookedPackagesMap = MoveTemp(RemainingUncookedPackagesMap);
+	return Count;
+}
+#endif //if WITH_EDITOR
+
 
 //PRAGMA_ENABLE_OPTIMIZATION
\ No newline at end of file
diff --git a/Engine/Source/Runtime/PakFile/Private/FilePackageStore.h b/Engine/Source/Runtime/PakFile/Private/FilePackageStore.h
--- a/Engine/Source/Runtime/PakFile/Private/FilePackageStore.h
+++ b/Engine/Source/Runtime/PakFile/Private/FilePackageStore.h
@@ -44,12 +44,13 @@ private:
 		FName PackageName;
 		EPackageExtension HeaderExtension;
 	};
-#endif
+#endif //if WITH_EDITOR
 
 	void Update();
 #if WITH_EDITOR
 	uint64 AddUncookedPackagesFromRoot(const FString& RootPath);
-#endif
+	uint64 RemoveUncookedPackagesFromRoot(const TSet<FString>& RootPath);
+#endif //if WITH_EDITOR
 
 	FRWLock EntriesLock;
 	FCriticalSection UpdateLock;
@@ -62,13 +63,13 @@ private:
 
 #if WITH_EDITOR
 	FDelegateHandle OnContentPathMountedDelegateHandle;
+	FDelegateHandle OnContentPathDismountedDelegateHandle;
 	FCriticalSection UncookedPackageRootsLock;
-	TSet<FString> PendingUncookedPackageRoots;
+	TSet<FString> PendingAddUncookedPackageRoots;
+	TSet<FString> PendingRemoveUncookedPackageRoots;
 	TMap<FPackageId, FUncookedPackage> UncookedPackagesMap;
-#endif
-#if WITH_EDITOR
 	TMap<FPackageId, const FFilePackageStoreEntry*> OptionalSegmentStoreEntriesMap;
-#endif
+#endif //if WITH_EDITOR
 
 	static thread_local int32 LockedOnThreadCount;
 };
