Fixes -cvarsini when bAllowNonUFSIniWhenCooked=false

diff --git a/Engine/Source/Runtime/PakFile/Private/IPlatformFilePak.cpp b/Engine/Source/Runtime/PakFile/Private/IPlatformFilePak.cpp
--- a/Engine/Source/Runtime/PakFile/Private/IPlatformFilePak.cpp
+++ b/Engine/Source/Runtime/PakFile/Private/IPlatformFilePak.cpp
@@ -5172,6 +5172,18 @@ bool FPakPlatformFile::IsNonPakFilenameAllowed(const FString& InFilename)
 		}
 	}
 #endif
+#if !DISABLE_CHEAT_CVARS && !UE_BUILD_SHIPPING
+	if (bIsIniFile && !bAllowed)
+	{
+		FString OverrideConsoleVariablesPath;
+		FParse::Value(FCommandLine::Get(), TEXT("-cvarsini="), OverrideConsoleVariablesPath);
+
+		if (!OverrideConsoleVariablesPath.IsEmpty() && InFilename == OverrideConsoleVariablesPath)
+		{
+			bAllowed = true;
+		}
+	}
+#endif
 
 	FFilenameSecurityDelegate& FilenameSecurityDelegate = GetFilenameSecurityDelegate();
 	if (bAllowed)
