PR #9165: Prevent fileOpenLog command-line to bypass signature check (Contributed by [REDACTED_PERSON])

diff --git a/Engine/Source/Runtime/PakFile/Private/IPlatformFilePak.cpp b/Engine/Source/Runtime/PakFile/Private/IPlatformFilePak.cpp
--- a/Engine/Source/Runtime/PakFile/Private/IPlatformFilePak.cpp
+++ b/Engine/Source/Runtime/PakFile/Private/IPlatformFilePak.cpp
@@ -7434,8 +7434,11 @@ bool FPakPlatformFile::Initialize(IPlatformFile* Inner, const TCHAR* CmdLine)
 	GameUserSettingsIniFilename = TEXT("GameUserSettings.ini");
 #endif
 
-	// Signed if we have keys, and are not running with fileopenlog (currently results in a deadlock).
-	bSigned = FCoreDelegates::GetPakSigningKeysDelegate().IsBound() && !FParse::Param(FCommandLine::Get(), TEXT("fileopenlog"));
+	// Signed if we have keys, and are not running with fileopenlog in non-shipping builds (currently results in a deadlock).
+	bSigned = FCoreDelegates::GetPakSigningKeysDelegate().IsBound();
+#if !UE_BUILD_SHIPPING
+	bSigned &= !FParse::Param(FCommandLine::Get(), TEXT("fileopenlog"));
+#endif
 
 	FString StartupPaksWildcard = GMountStartupPaksWildCard;
 #if !UE_BUILD_SHIPPING
