Allow for platform specializations of the 'is pak file installed' logic when enumerating pak files.

diff --git a/Engine/Source/Runtime/PakFile/Private/IPlatformFilePak.cpp b/Engine/Source/Runtime/PakFile/Private/IPlatformFilePak.cpp
--- a/Engine/Source/Runtime/PakFile/Private/IPlatformFilePak.cpp
+++ b/Engine/Source/Runtime/PakFile/Private/IPlatformFilePak.cpp
@@ -68,6 +68,9 @@
 #ifndef ALLOW_INI_OVERRIDE_FROM_COMMANDLINE
 #define ALLOW_INI_OVERRIDE_FROM_COMMANDLINE 0
 #endif
+#ifndef HAS_PLATFORM_PAK_INSTALL_CHECK
+#define HAS_PLATFORM_PAK_INSTALL_CHECK 0
+#endif
 #ifndef ALL_PAKS_WILDCARD
 #define ALL_PAKS_WILDCARD "*.pak"
 #endif 
@@ -5197,7 +5200,29 @@
 	return bAllowed;
 }
 
+#if !HAS_PLATFORM_PAK_INSTALL_CHECK
+bool FPakPlatformFile::IsPakFileInstalled(const FString& InFilename)
+{
+	IPlatformChunkInstall* ChunkInstall = FPlatformMisc::GetPlatformChunkInstall();
+	if (ChunkInstall)
+	{
+		// if a platform supports chunk style installs, make sure that the chunk a pak file resides in is actually fully installed before accepting pak files from it
+		int32 PakchunkIndex = GetPakchunkIndexFromPakFile(InFilename);
+		if (PakchunkIndex != INDEX_NONE)
+		{
+			if (ChunkInstall->GetPakchunkLocation(PakchunkIndex) == EChunkLocation::NotAvailable)
+			{
+				return false;
+			}
+		}
+	}
 
+	return true;
+}
+#endif //HAS_PLATFORM_PAK_INSTALL_CHECK
+
+
+
 FSharedPakReader::FSharedPakReader(FArchive* InArchive, FPakFile* InPakFile)
 	: Archive(InArchive)
 	, PakFile(InPakFile)
@@ -7290,14 +7315,12 @@
 	class FPakSearchVisitor : public IPlatformFile::FDirectoryVisitor
 	{
 		TArray<FString>& FoundPakFiles;
-		IPlatformChunkInstall* ChunkInstall = nullptr;
 		FString WildCard;
 		bool bSkipOptionalPakFiles;
 
 	public:
-		FPakSearchVisitor(TArray<FString>& InFoundPakFiles, const FString& InWildCard, IPlatformChunkInstall* InChunkInstall, bool bInSkipOptionalPakFiles)
+		FPakSearchVisitor(TArray<FString>& InFoundPakFiles, const FString& InWildCard, bool bInSkipOptionalPakFiles)
 			: FoundPakFiles(InFoundPakFiles)
-			, ChunkInstall(InChunkInstall)
 			, WildCard(InWildCard)
 			, bSkipOptionalPakFiles(bInSkipOptionalPakFiles)
 		{}
@@ -7308,17 +7331,9 @@
 				FString Filename(FilenameOrDirectory);
 				if(Filename.MatchesWildcard(WildCard))
 				{
-					// if a platform supports chunk style installs, make sure that the chunk a pak file resides in is actually fully installed before accepting pak files from it
-					if (ChunkInstall)
+					if (!FPakPlatformFile::IsPakFileInstalled(Filename))
 					{
-						int32 PakchunkIndex = GetPakchunkIndexFromPakFile(Filename);
-						if (PakchunkIndex != INDEX_NONE)
-						{
-							if (ChunkInstall->GetPakchunkLocation(PakchunkIndex) == EChunkLocation::NotAvailable)
-							{
-								return true;
-							}
-						}
+						return true;
 					}
 
 #if !UE_BUILD_SHIPPING
@@ -7336,7 +7351,7 @@
 	bool bSkipOptionalPakFiles = FParse::Param(FCommandLine::Get(), TEXT("SkipOptionalPakFiles"));
 
 	// Find all pak files.
-	FPakSearchVisitor Visitor(OutPakFiles, WildCard, FPlatformMisc::GetPlatformChunkInstall(), bSkipOptionalPakFiles);
+	FPakSearchVisitor Visitor(OutPakFiles, WildCard, bSkipOptionalPakFiles);
 	LowLevelFile->IterateDirectoryRecursively(Directory, Visitor);
 }
 
diff --git a/Engine/Source/Runtime/PakFile/Public/IPlatformFilePak.h b/Engine/Source/Runtime/PakFile/Public/IPlatformFilePak.h
--- a/Engine/Source/Runtime/PakFile/Public/IPlatformFilePak.h
+++ b/Engine/Source/Runtime/PakFile/Public/IPlatformFilePak.h
@@ -2162,6 +2162,14 @@
 	 */
 	void RegisterEncryptionKey(const FGuid& InEncryptionKeyGuid, const FAES::FAESKey& InKey);
 
+	/**
+	 * Checks with any current chunk installation system if the given pak file is installed
+	 * 
+	 * @param InFilename  the pak filename to check
+	 * @return whether the pak file is installed
+	 */
+	static bool IsPakFileInstalled(const FString& InFilename);
+
 public:
 
 	//~ For visibility of overloads we don't override
