FilePackageStore: Fixed reversed ordering of mounted containers

diff --git a/Engine/Source/Runtime/PakFile/Private/FilePackageStore.cpp b/Engine/Source/Runtime/PakFile/Private/FilePackageStore.cpp
--- a/Engine/Source/Runtime/PakFile/Private/FilePackageStore.cpp
+++ b/Engine/Source/Runtime/PakFile/Private/FilePackageStore.cpp
@@ -161,9 +161,9 @@ void FFilePackageStoreBackend::Mount(const FIoContainerHeader* ContainerHeader,
 		{
 			if (A.Order == B.Order)
 			{
-				return A.Sequence < B.Sequence;
+				return A.Sequence > B.Sequence;
 			}
-			return A.Order < B.Order;
+			return A.Order > B.Order;
 		});
 	bNeedsUpdate = true;
 }
@@ -244,9 +244,14 @@ void FFilePackageStoreBackend::Update()
 
 		for (const FIoContainerHeaderPackageRedirect& Redirect : ContainerHeader->PackageRedirects)
 		{
-			FDisplayNameEntryId NameEntry = ContainerHeader->RedirectsNameMap[Redirect.SourcePackageName.GetIndex()];
-			FName SourcePackageName = NameEntry.ToName(Redirect.SourcePackageName.GetNumber());
-			RedirectsPackageMap.Emplace(Redirect.SourcePackageId, MakeTuple(SourcePackageName, Redirect.TargetPackageId));
+			TTuple<FName, FPackageId>& RedirectEntry = RedirectsPackageMap.FindOrAdd(Redirect.SourcePackageId);
+			FName& SourcePackageName = RedirectEntry.Key;
+			if (SourcePackageName.IsNone())
+			{
+				FDisplayNameEntryId NameEntry = ContainerHeader->RedirectsNameMap[Redirect.SourcePackageName.GetIndex()];
+				SourcePackageName = NameEntry.ToName(Redirect.SourcePackageName.GetNumber());
+				RedirectEntry.Value = Redirect.TargetPackageId;
+			}
 		}
 	}
 
