fixed TSAN warning claiming a data race concurrently accessing `FFileIoStoreStats::AvailableBuffersCounters`

diff --git a/Engine/Source/Runtime/PakFile/Internal/IoDispatcherFileBackendTypes.h b/Engine/Source/Runtime/PakFile/Internal/IoDispatcherFileBackendTypes.h
--- a/Engine/Source/Runtime/PakFile/Internal/IoDispatcherFileBackendTypes.h
+++ b/Engine/Source/Runtime/PakFile/Internal/IoDispatcherFileBackendTypes.h
@@ -936,7 +936,7 @@ private:
 	FCountersTrace::FCounterInt BlockCacheMissedSizeCounter;
 	FCountersTrace::FCounterInt ScatteredSizeCounter;
 	FCountersTrace::FCounterInt TocMemoryCounter;
-	FCountersTrace::FCounterInt AvailableBuffersCounter;
+	FCountersTrace::TCounter<std::atomic<int64>, TraceCounterType_Int> AvailableBuffersCounter;
 #endif
 
 #if CSV_PROFILER
