[Engine Runtime]
* IPlatformFilePak - Removed tons of string creation/destruction when finding files in path. Saves lots of seconds when starting dedicated server where plugins are loaded in sync

Saves some seconds during early startup primarily in dedicated server where plugins are not created in parallel

diff --git a/Engine/Source/Runtime/PakFile/Public/IPlatformFilePak.h b/Engine/Source/Runtime/PakFile/Public/IPlatformFilePak.h
--- a/Engine/Source/Runtime/PakFile/Public/IPlatformFilePak.h
+++ b/Engine/Source/Runtime/PakFile/Public/IPlatformFilePak.h
@@ -1728,13 +1728,21 @@ private:
 	template <class ContainerType>
 	void FindFilesAtPathInIndex(const FDirectoryIndex& TargetIndex, ContainerType& OutFiles, const FString& Directory, bool bIncludeFiles = true, bool bIncludeDirectories = false, bool bRecursive = false) const
 	{
+		// Early out if MountPoint is not matching directory
+		if (!Directory.StartsWith(MountPoint))
+		{
+			return;
+		}
+
+		FStringView RelativeSearch(FStringView(Directory).RightChop(MountPoint.Len()));
+
 		TArray<FString> DirectoriesInPak; // List of all unique directories at path
 		for (TMap<FString, FPakDirectory>::TConstIterator It(TargetIndex); It; ++It)
 		{
-			FString PakPath(PakPathCombine(MountPoint, It.Key()));
 			// Check if the file is under the specified path.
-			if (PakPath.StartsWith(Directory))
+			if (FStringView(It.Key()).StartsWith(RelativeSearch))
 			{
+				FString PakPath = PakPathCombine(MountPoint, It.Key());
 				if (bRecursive == true)
 				{
 					// Add everything
