Removed unacceptable words from IPlatformPak

diff --git a/Engine/Source/Runtime/PakFile/Private/IPlatformFilePak.cpp b/Engine/Source/Runtime/PakFile/Private/IPlatformFilePak.cpp
--- a/Engine/Source/Runtime/PakFile/Private/IPlatformFilePak.cpp
+++ b/Engine/Source/Runtime/PakFile/Private/IPlatformFilePak.cpp
@@ -225,30 +225,27 @@ FPakPlatformFile::FPakSigningFailureHandlerData& FPakPlatformFile::GetPakSigning
 	return Instance;
 }
 
-// Needs to be a global as multiple threads may try to access the failed delegate
-FPakChunkSignatureCheckFailedHandler& FPakPlatformFile::GetPakChunkSignatureCheckFailedHandler()
-{
-	return GetPakSigningFailureHandlerData().ChunkSignatureCheckFailedDelegate;
-}
-
-FPakMasterSignatureTableCheckFailureHandler& FPakPlatformFile::GetPakMasterSignatureTableCheckFailureHandler()
+void FPakPlatformFile::BroadcastPakChunkSignatureCheckFailure(const FPakChunkSignatureCheckFailedData& InData)
 {
-	return GetPakSigningFailureHandlerData().MasterSignatureTableCheckFailedDelegate;
+	FPakSigningFailureHandlerData& HandlerData = GetPakSigningFailureHandlerData();
+	FScopeLock Lock(&HandlerData.GetLock());
+	HandlerData.GetPakChunkSignatureCheckFailedDelegate().Broadcast(InData);
 }
 
-void FPakPlatformFile::BroadcastPakChunkSignatureCheckFailure(const FPakChunkSignatureCheckFailedData& InData)
+void FPakPlatformFile::BroadcastPakPrincipalSignatureTableCheckFailure(const FString& InFilename)
 {
 	FPakSigningFailureHandlerData& HandlerData = GetPakSigningFailureHandlerData();
-	FScopeLock Lock(&HandlerData.Lock);
-	HandlerData.ChunkSignatureCheckFailedDelegate.Broadcast(InData);
+	FScopeLock Lock(&HandlerData.GetLock());
+	HandlerData.GetPrincipalSignatureTableCheckFailedDelegate().Broadcast(InFilename);
 }
 
+PRAGMA_DISABLE_DEPRECATION_WARNINGS
 void FPakPlatformFile::BroadcastPakMasterSignatureTableCheckFailure(const FString& InFilename)
 {
-	FPakSigningFailureHandlerData& HandlerData = GetPakSigningFailureHandlerData();
-	FScopeLock Lock(&HandlerData.Lock);
-	HandlerData.MasterSignatureTableCheckFailedDelegate.Broadcast(InFilename);
+	return BroadcastPakPrincipalSignatureTableCheckFailure(InFilename);
 }
+PRAGMA_ENABLE_DEPRECATION_WARNINGS
+
 
 FPakSetIndexSettings& FPakPlatformFile::GetPakSetIndexSettingsDelegate()
 {
@@ -411,7 +408,7 @@ TSharedPtr<const struct FPakSignatureFile, ESPMode::ThreadSafe> FPakPlatformFile
 		else
 		{
 			UE_LOG(LogPakFile, Warning, TEXT("Couldn't find pak signature file '%s'"), InFilename);
-			BroadcastPakMasterSignatureTableCheckFailure(InFilename);
+			BroadcastPakPrincipalSignatureTableCheckFailure(InFilename);
 		}
 	}
 
@@ -3749,7 +3746,7 @@ void FPakPrecacher::DoSignatureCheck(bool bWasCanceled, IAsyncReadRequest* Reque
 	int64 RequestSize = 0;
 	int64 RequestOffset = 0;
 	uint16 PakIndex;
-	FSHAHash MasterSignatureHash;
+	FSHAHash PrincipalSignatureHash;
 	static const int64 MaxHashesToCache = 16;
 
 #if PAKHASH_USE_CRC
@@ -3777,7 +3774,7 @@ void FPakPrecacher::DoSignatureCheck(bool bWasCanceled, IAsyncReadRequest* Reque
 		SignatureIndex = RequestOffset / FPakInfo::MaxChunkDataSize;
 
 		FPakData& PakData = CachedPakData[PakIndex];
-		MasterSignatureHash = PakData.Signatures->DecryptedHash;
+		PrincipalSignatureHash = PakData.Signatures->DecryptedHash;
 
 		for (int32 CacheIndex = 0; CacheIndex < FMath::Min(NumSignaturesToCheck, MaxHashesToCache); ++CacheIndex)
 		{
@@ -3819,9 +3816,9 @@ void FPakPrecacher::DoSignatureCheck(bool bWasCanceled, IAsyncReadRequest* Reque
 				UE_LOG(LogPakFile, Warning, TEXT("Pak chunk signing mismatch on chunk [%i/%i]! Expected %s, Received %s"), SignatureIndex, PakData->Signatures->ChunkHashes.Num() - 1, *ChunkHashToString(PakData->Signatures->ChunkHashes[SignatureIndex]), *ChunkHashToString(ThisHash));
 
 				// Check the signatures are still as we expected them
-				if (PakData->Signatures->DecryptedHash != PakData->Signatures->ComputeCurrentMasterHash())
+				if (PakData->Signatures->DecryptedHash != PakData->Signatures->ComputeCurrentPrincipalHash())
 				{
-					UE_LOG(LogPakFile, Warning, TEXT("Master signature table has changed since initialization!"));
+					UE_LOG(LogPakFile, Warning, TEXT("Principal signature table has changed since initialization!"));
 				}
 
 				FPakChunkSignatureCheckFailedData FailedData(PakData->Name.ToString(), HashCache[SignedChunkIndex % MaxHashesToCache], ThisHash, SignatureIndex);
@@ -6687,8 +6684,8 @@ bool FPakFile::Check()
 		FPakPlatformFile::FPakSigningFailureHandlerData& HandlerData = FPakPlatformFile::GetPakSigningFailureHandlerData();
 
 		{
-			FScopeLock Lock(&HandlerData.Lock);
-			DelegateHandle = HandlerData.ChunkSignatureCheckFailedDelegate.AddLambda([&ErrorCount](const FPakChunkSignatureCheckFailedData&)
+			FScopeLock Lock(&HandlerData.GetLock());
+			DelegateHandle = HandlerData.GetPakChunkSignatureCheckFailedDelegate().AddLambda([&ErrorCount](const FPakChunkSignatureCheckFailedData&)
 			{
 				++ErrorCount;
 			});
@@ -6706,8 +6703,8 @@ bool FPakFile::Check()
 
 		if (DelegateHandle.IsValid())
 		{
-			FScopeLock Lock(&HandlerData.Lock);
-			HandlerData.ChunkSignatureCheckFailedDelegate.Remove(DelegateHandle);
+			FScopeLock Lock(&HandlerData.GetLock());
+			HandlerData.GetPakChunkSignatureCheckFailedDelegate().Remove(DelegateHandle);
 		}
 	}
 	else
diff --git a/Engine/Source/Runtime/PakFile/Private/SignedArchiveReader.cpp b/Engine/Source/Runtime/PakFile/Private/SignedArchiveReader.cpp
--- a/Engine/Source/Runtime/PakFile/Private/SignedArchiveReader.cpp
+++ b/Engine/Source/Runtime/PakFile/Private/SignedArchiveReader.cpp
@@ -267,9 +267,9 @@ bool FChunkCacheWorker::CheckSignature(const FChunkRequest& ChunkInfo)
 		{
 			UE_LOG(LogPakFile, Warning, TEXT("Pak chunk signing mismatch on chunk [%i/%i]! Expected %s, Received %s"), ChunkInfo.Index, Signatures->ChunkHashes.Num() - 1, *ChunkHashToString(Signatures->ChunkHashes[ChunkInfo.Index]), *ChunkHashToString(ChunkHash));
 
-			if (Signatures->DecryptedHash != Signatures->ComputeCurrentMasterHash())
+			if (Signatures->DecryptedHash != Signatures->ComputeCurrentPrincipalHash())
 			{
-				UE_LOG(LogPakFile, Warning, TEXT("Master signature table has changed since initialization!"));
+				UE_LOG(LogPakFile, Warning, TEXT("Principal signature table has changed since initialization!"));
 			}
 
 			const FPakChunkSignatureCheckFailedData Data(Reader->GetArchiveName(), Signatures->ChunkHashes[ChunkInfo.Index], ChunkHash, ChunkInfo.Index);
diff --git a/Engine/Source/Runtime/PakFile/Private/SignedArchiveReader.h b/Engine/Source/Runtime/PakFile/Private/SignedArchiveReader.h
--- a/Engine/Source/Runtime/PakFile/Private/SignedArchiveReader.h
+++ b/Engine/Source/Runtime/PakFile/Private/SignedArchiveReader.h
@@ -206,7 +206,7 @@ public:
 
 	/**
 	* Indicates that this chunk worker is valid. If the signature file couldn't be loaded or if it failed
-	* the master table check, this will be false
+	* the principal table check, this will be false
 	*/
 	bool IsValid() const;
 
diff --git a/Engine/Source/Runtime/PakFile/Public/IPlatformFilePak.h b/Engine/Source/Runtime/PakFile/Public/IPlatformFilePak.h
--- a/Engine/Source/Runtime/PakFile/Public/IPlatformFilePak.h
+++ b/Engine/Source/Runtime/PakFile/Public/IPlatformFilePak.h
@@ -109,10 +109,12 @@ struct FPakChunkSignatureCheckFailedData
 DECLARE_DELEGATE_RetVal_OneParam(bool, FFilenameSecurityDelegate, const TCHAR* /*InFilename*/);
 DECLARE_DELEGATE_ThreeParams(FPakCustomEncryptionDelegate, uint8* /*InData*/, uint32 /*InDataSize*/, FGuid /*InEncryptionKeyGuid*/);
 DECLARE_MULTICAST_DELEGATE_OneParam(FPakChunkSignatureCheckFailedHandler, const FPakChunkSignatureCheckFailedData&);
-DECLARE_MULTICAST_DELEGATE_OneParam(FPakMasterSignatureTableCheckFailureHandler, const FString&);
+DECLARE_MULTICAST_DELEGATE_OneParam(FPakPrincipalSignatureTableCheckFailureHandler, const FString&);
 /** Delegate which allows a project to configure Index Pruning.  This is a delegate instead of a config file because config files are loaded after the first PakFiles */
 DECLARE_DELEGATE_ThreeParams(FPakSetIndexSettings, bool& /* bKeepFullDirectory*/, bool& /* bValidatePruning */, bool& /* bDelayPruning */);
 
+UE_DEPRECATED("5.1", "Use FPakPrincipalSignatureTableCheckFailureHandler instead")
+typedef FPakPrincipalSignatureTableCheckFailureHandler FPakMasterSignatureTableCheckFailureHandler;
 /**
  * Struct which holds pak file info (version, index offset, hash value).
  */
@@ -965,17 +967,6 @@ public:
 		return CachedTotalSize;
 	}
 
-	/**
-	 * Gets pak file index.
-	 *
-	 * @return Pak index.
-	 */
-	UE_DEPRECATED(4.26, "Use FPrunedFilenameIterator or FPakEntryIterator instead; the Index is now no longer necessarily a DirectoryIndex and in some cases requires a lock around Index use")
-	const TMap<FString, FPakDirectory>& GetIndex() const
-	{
-		return DirectoryIndex;
-	}
-
 	/**
 	 * Gets the number of files in this pak.
 	 */
@@ -984,21 +975,9 @@ public:
 		return NumEntries;
 	}
 
-	UE_DEPRECATED(4.26, "Use GetPrunedFilenames instead")
-	void GetFilenames(TArray<FString>& OutFileList) const
-	{
-		return GetPrunedFilenames(OutFileList);
-	}
-
 	/** Returns the FullPath (includes Mount) Filename found in Pruned DirectoryIndex */
 	void GetPrunedFilenames(TArray<FString>& OutFileList) const;
 
-	UE_DEPRECATED(4.26, "Use GetPrunedFilenamesInChunk instead")
-	void GetFilenamesInChunk(const TArray<int32>& InChunkIDs, TArray<FString>& OutFileList) const
-	{
-		GetPrunedFilenamesInChunk(InChunkIDs, OutFileList);
-	}
-
 	/** Returns the RelativePathFromMount Filename for every Filename found in the Pruned DirectoryIndex that points to a PakEntry in the given Chunk */
 	void GetPrunedFilenamesInChunk(const TArray<int32>& InChunkIDs, TArray<FString>& OutFileList) const;
 
@@ -1051,13 +1030,6 @@ public:
 		return MountPoint;
 	}
 
-	template <class ContainerType>
-	UE_DEPRECATED(4.26, "Use FindPrunedFilesAtPath instead")
-	void FindFilesAtPath(ContainerType& OutFiles, const TCHAR* InPath, bool bIncludeFiles = true, bool bIncludeDirectories = false, bool bRecursive = false) const
-	{
-		FindPrunedFilesAtPath(OutFiles, InPath, bIncludeFiles, bIncludeDirectories, bRecursive);
-	}
-
 	/**
 	 * Looks for files or directories within the Pruned DirectoryIndex of the pak file.
 	 * The Pruned DirectoryIndex does not have entries for most Files in the pak; they were removed to save memory.
@@ -1112,12 +1084,6 @@ public:
 	 * @param InPath Directory path.
 	 * @return Pointer to a map with directory contents if the directory was found, NULL otherwise.
 	 */
-	UE_DEPRECATED(4.26, "Use FindPrunedDirectory instead and wrap the access with FScopedPakDirectoryIndexAccess; const FPakDirectory is no longer threadsafe.")
-	const FPakDirectory* FindDirectory(const TCHAR* InPath) const
-	{
-		return nullptr;
-	}
-
 	const FPakDirectory* FindPrunedDirectory(const TCHAR* InPath) const
 	{
 		// Caller holds FScopedPakDirectoryIndexAccess
@@ -1130,12 +1096,6 @@ public:
 		return FindPrunedDirectoryInternal(RelativePathFromMount);
 	}
 
-	UE_DEPRECATED(4.26, "Use DirectoryExistsInPruned instead.")
-	bool DirectoryExists(const TCHAR* InPath) const
-	{
-		return DirectoryExistsInPruned(InPath);
-	}
-
 	/**
 	 * Checks if a directory exists in pak file.
 	 *
@@ -1908,7 +1868,7 @@ class PAKFILE_API FPakFileHandle : public IFileHandle
 public:
 
 	UE_DEPRECATED(4.27, "Use constructor that takes a TRefCountPtr<FPakFile> instead")
-	FPakFileHandle(const FPakFile& InPakFile, const FPakEntry& InPakEntry, TAcquirePakReaderFunction& InAcquirePakReaderFunction)
+		FPakFileHandle(const FPakFile& InPakFile, const FPakEntry& InPakEntry, TAcquirePakReaderFunction& InAcquirePakReaderFunction)
 		: FPakFileHandle(TRefCountPtr<const FPakFile>(&InPakFile), InPakEntry, InAcquirePakReaderFunction)
 	{
 	}
@@ -1929,7 +1889,7 @@ public:
 	}
 
 	UE_DEPRECATED(4.27, "Use constructor that takes a TRefCountPtr<FPakFile> instead")
-	FPakFileHandle(const FPakFile& InPakFile, const FPakEntry& InPakEntry, FArchive* InPakReader)
+		FPakFileHandle(const FPakFile& InPakFile, const FPakEntry& InPakEntry, FArchive* InPakReader)
 		: FPakFileHandle(TRefCountPtr<const FPakFile>(&InPakFile), InPakEntry, InPakReader)
 	{
 	}
@@ -3107,44 +3067,39 @@ public:
 
 	struct FPakSigningFailureHandlerData
 	{
+		PRAGMA_DISABLE_DEPRECATION_WARNINGS
+		FCriticalSection& GetLock() { return Lock; }
+		FPakChunkSignatureCheckFailedHandler& GetPakChunkSignatureCheckFailedDelegate() { return ChunkSignatureCheckFailedDelegate; }
+		FPakPrincipalSignatureTableCheckFailureHandler& GetPrincipalSignatureTableCheckFailedDelegate() { return MasterSignatureTableCheckFailedDelegate; }
+		PRAGMA_ENABLE_DEPRECATION_WARNINGS
+
+		UE_DEPRECATED("5.1", "Use GetLock instead")
 		FCriticalSection Lock;
+
+		UE_DEPRECATED("5.1", "Use GetPakChunkSignatureCheckFailedDelegate instead")
 		FPakChunkSignatureCheckFailedHandler ChunkSignatureCheckFailedDelegate;
-		FPakMasterSignatureTableCheckFailureHandler MasterSignatureTableCheckFailedDelegate;
+
+		UE_DEPRECATED("5.1", "Use GetPrincipalSignatureTableCheckFailureDelegate instead")
+		FPakPrincipalSignatureTableCheckFailureHandler MasterSignatureTableCheckFailedDelegate;
 	};
 
-	// Access static delegate for handling a pak signature check failure
+	// Access static delegate for handling a Pak signature check failure
 	static FPakSigningFailureHandlerData& GetPakSigningFailureHandlerData();
 	
-	// Access static delegate for handling a pak signature check failure
-	UE_DEPRECATED(4.25, "GetPakChunkSignatureCheckFailedHandler is not thread safe, so please migrate to using GetPakSigningFailureHandlerData and locking the critical section around any use of the delegates")
-	static FPakChunkSignatureCheckFailedHandler& GetPakChunkSignatureCheckFailedHandler();
-
-	// Access static delegate for handling a pak signature check failure
-	UE_DEPRECATED(4.25, "GetPakMasterSignatureTableCheckFailureHandler is not thread safe, so please migrate to using GetPakSigningFailureHandlerData and locking the critical section around any use of the delegates")
-	static FPakMasterSignatureTableCheckFailureHandler& GetPakMasterSignatureTableCheckFailureHandler();
-
-	// Broadacast a signature check failure through any registered delegates in a thread safe way
+	// Broadcast a signature check failure through any registered delegates in a thread safe way
 	static void BroadcastPakChunkSignatureCheckFailure(const FPakChunkSignatureCheckFailedData& InData);
 
-	// Broadacast a master signature table failure through any registered delegates in a thread safe way
+	// Broadcast a principal signature table failure through any registered delegates in a thread safe way
+	static void BroadcastPakPrincipalSignatureTableCheckFailure(const FString& InFilename);
+
+	UE_DEPRECATED("5.1", "Use BroadcastPakPrincipalSignatureTableCheckFailure instead")
 	static void BroadcastPakMasterSignatureTableCheckFailure(const FString& InFilename);
 
 	// Access static delegate for setting PakIndex settings.
 	static FPakSetIndexSettings& GetPakSetIndexSettingsDelegate();
 
-	UE_DEPRECATED(4.26, "Use GetPrunedFilenamesInChunk")
-	void GetFilenamesInChunk(const FString& InPakFilename, const TArray<int32>& InChunkIDs, TArray<FString>& OutFileList)
-	{
-		GetPrunedFilenamesInChunk(InPakFilename, InChunkIDs, OutFileList);
-	}
-
 	/* Get a list of RelativePathFromMount for every file in the given Pak that lives in any of the given chunks.  Only searches the Pruned DirectoryIndex */
 	void GetPrunedFilenamesInChunk(const FString& InPakFilename, const TArray<int32>& InChunkIDs, TArray<FString>& OutFileList);
-	UE_DEPRECATED(4.26, "Use GetPrunedFilenamesInPakFile")
-	void GetFilenamesInPakFile(const FString& InPakFilename, TArray<FString>& OutFileList)
-	{
-		GetPrunedFilenamesInPakFile(InPakFilename, OutFileList);
-	}
 
 	/** Gets a list of FullPaths (includes Mount directory) for every File in the given Pak's Pruned DirectoryIndex */
 	void GetPrunedFilenamesInPakFile(const FString& InPakFilename, TArray<FString>& OutFileList);
@@ -3217,7 +3172,7 @@ struct FPakSignatureFile
 	{
 		ChunkHashes = InChunkHashes;
 		SignatureData = InSignatureData;
-		DecryptedHash = ComputeCurrentMasterHash();
+		DecryptedHash = ComputeCurrentPrincipalHash();
 
 		TArray<uint8> NewSignatureData;
 		NewSignatureData.Append(SignatureData);
@@ -3262,7 +3217,7 @@ struct FPakSignatureFile
 			{
 				FMemory::Memcpy(DecryptedHash.Hash, SignatureData.GetData() + SignatureData.Num() - UE_ARRAY_COUNT(FSHAHash::Hash), UE_ARRAY_COUNT(FSHAHash::Hash));
 				SignatureData.SetNum(SignatureData.Num() - UE_ARRAY_COUNT(FSHAHash::Hash));
-				FSHAHash CurrentHash = ComputeCurrentMasterHash();
+				FSHAHash CurrentHash = ComputeCurrentPrincipalHash();
 				if (DecryptedHash == CurrentHash)
 				{
 					return true;
@@ -3278,17 +3233,23 @@ struct FPakSignatureFile
 			}
 		}
 
-		FPakPlatformFile::BroadcastPakMasterSignatureTableCheckFailure(InFilename);
+		FPakPlatformFile::BroadcastPakPrincipalSignatureTableCheckFailure(InFilename);
 		return false;
 	}
 
 	/**
 	 * Helper function for computing the SHA1 hash of the current chunk CRC array
 	 */
-	FSHAHash ComputeCurrentMasterHash() const
+	FSHAHash ComputeCurrentPrincipalHash() const
 	{
 		FSHAHash CurrentHash;
 		FSHA1::HashBuffer(ChunkHashes.GetData(), ChunkHashes.Num() * sizeof(TPakChunkHash), CurrentHash.Hash);
 		return CurrentHash;
 	}
+
+	UE_DEPRECATED("5.1", "Use ComputeCurrentPrincipalHash instead")
+	FSHAHash ComputeCurrentMasterHash() const
+	{
+		return ComputeCurrentPrincipalHash();
+	}
 };
