PAK - use FileExists instead of OpenRead when mounting PAK files

#8465
#rnx

diff --git a/Engine/Source/Runtime/PakFile/Private/IPlatformFilePak.cpp b/Engine/Source/Runtime/PakFile/Private/IPlatformFilePak.cpp
--- a/Engine/Source/Runtime/PakFile/Private/IPlatformFilePak.cpp
+++ b/Engine/Source/Runtime/PakFile/Private/IPlatformFilePak.cpp
@@ -7680,8 +7680,7 @@ bool FPakPlatformFile::Mount(const TCHAR* InPakFilename, uint32 PakOrder, const
 	LLM_SCOPE(ELLMTag::FileSystem);
 	bool bPakSuccess = false;
 	bool bIoStoreSuccess = true;
-	TSharedPtr<IFileHandle> PakHandle = MakeShareable(LowerLevel->OpenRead(InPakFilename));
-	if (PakHandle.IsValid())
+	if (LowerLevel->FileExists(InPakFilename))
 	{
 		TRefCountPtr<FPakFile> Pak = new FPakFile(LowerLevel, InPakFilename, bSigned, bLoadIndex);
 		if (Pak.GetReference()->IsValid())
@@ -7745,7 +7744,6 @@ bool FPakPlatformFile::Mount(const TCHAR* InPakFilename, uint32 PakOrder, const
 				Entry.PakchunkIndex = Pak->PakchunkIndex;
 
 				Pak.SafeRelease();
-				PakHandle.Reset();
 				return false;
 			}
 		}
