FilePackageStore (editor): Don't update mounted iostore containers if we've only added or removed paths for uncooked packages

diff --git a/Engine/Source/Runtime/PakFile/Private/FilePackageStore.cpp b/Engine/Source/Runtime/PakFile/Private/FilePackageStore.cpp
--- a/Engine/Source/Runtime/PakFile/Private/FilePackageStore.cpp
+++ b/Engine/Source/Runtime/PakFile/Private/FilePackageStore.cpp
@@ -38,7 +38,7 @@ FFilePackageStoreBackend::FFilePackageStoreBackend()
 			}
 			{
 				FWriteScopeLock _(EntriesLock);
-				bNeedsUpdate = true;
+				bNeedsUncookedPackagesUpdate = true;
 			}
 		});
 
@@ -51,7 +51,7 @@ FFilePackageStoreBackend::FFilePackageStoreBackend()
 			}
 			{
 				FWriteScopeLock _(EntriesLock);
-				bNeedsUpdate = true;
+				bNeedsUncookedPackagesUpdate = true;
 			}
 		});
 
@@ -63,7 +63,7 @@ FFilePackageStoreBackend::FFilePackageStoreBackend()
 	}
 	{
 		FWriteScopeLock _(EntriesLock);
-		bNeedsUpdate = true;
+		bNeedsUncookedPackagesUpdate = true;
 	}
 #endif //if WITH_EDITOR
 }
@@ -79,7 +79,11 @@ FFilePackageStoreBackend::~FFilePackageStoreBackend()
 void FFilePackageStoreBackend::BeginRead()
 {
 	EntriesLock.ReadLock();
-	if (bNeedsUpdate)
+	if (bNeedsContainerUpdate
+#if WITH_EDITOR
+		|| bNeedsUncookedPackagesUpdate
+#endif
+		)
 	{
 		Update();
 	}
@@ -165,7 +169,7 @@ void FFilePackageStoreBackend::Mount(const FIoContainerHeader* ContainerHeader,
 			}
 			return A.Order > B.Order;
 		});
-	bNeedsUpdate = true;
+	bNeedsContainerUpdate = true;
 }
 
 void FFilePackageStoreBackend::Unmount(const FIoContainerHeader* ContainerHeader)
@@ -176,7 +180,7 @@ void FFilePackageStoreBackend::Unmount(const FIoContainerHeader* ContainerHeader
 		if (It->ContainerHeader == ContainerHeader)
 		{
 			It.RemoveCurrent();
-			bNeedsUpdate = true;
+			bNeedsContainerUpdate = true;
 			return;
 		}
 	}
@@ -188,75 +192,76 @@ void FFilePackageStoreBackend::Update()
 	TRACE_CPUPROFILER_EVENT_SCOPE(UpdateFilePackageStore);
 
 	FScopeLock Lock(&UpdateLock);
-	if (!bNeedsUpdate)
+	if (bNeedsContainerUpdate)
 	{
-		return;
-	}
-
-	StoreEntriesMap.Empty();
-	LocalizedPackages.Empty();
-	RedirectsPackageMap.Empty();
+		StoreEntriesMap.Empty();
+		LocalizedPackages.Empty();
+		RedirectsPackageMap.Empty();
 #if WITH_EDITOR
-	OptionalSegmentStoreEntriesMap.Empty();
+		OptionalSegmentStoreEntriesMap.Empty();
 #endif
 
-	uint32 TotalPackageCount = 0;
-	for (const FMountedContainer& MountedContainer : MountedContainers)
-	{
-		TotalPackageCount += MountedContainer.ContainerHeader->PackageIds.Num();
-	}
-
-	StoreEntriesMap.Reserve(TotalPackageCount);
-	for (const FMountedContainer& MountedContainer : MountedContainers)
-	{
-		const FIoContainerHeader* ContainerHeader = MountedContainer.ContainerHeader;
-		TArrayView<const FFilePackageStoreEntry> ContainerStoreEntries(reinterpret_cast<const FFilePackageStoreEntry*>(ContainerHeader->StoreEntries.GetData()), ContainerHeader->PackageIds.Num());
-		int32 Index = 0;
-		for (const FFilePackageStoreEntry& StoreEntry : ContainerStoreEntries)
+		uint32 TotalPackageCount = 0;
+		for (const FMountedContainer& MountedContainer : MountedContainers)
 		{
-			const FPackageId& PackageId = ContainerHeader->PackageIds[Index];
-			check(PackageId.IsValid());
-			StoreEntriesMap.FindOrAdd(PackageId, &StoreEntry);
-			++Index;
+			TotalPackageCount += MountedContainer.ContainerHeader->PackageIds.Num();
 		}
 
-#if WITH_EDITOR
-		TArrayView<const FFilePackageStoreEntry> ContainerOptionalSegmentStoreEntries(reinterpret_cast<const FFilePackageStoreEntry*>(ContainerHeader->OptionalSegmentStoreEntries.GetData()), ContainerHeader->OptionalSegmentPackageIds.Num());
-		Index = 0;
-		for (const FFilePackageStoreEntry& OptionalSegmentStoreEntry : ContainerOptionalSegmentStoreEntries)
+		StoreEntriesMap.Reserve(TotalPackageCount);
+		for (const FMountedContainer& MountedContainer : MountedContainers)
 		{
-			const FPackageId& PackageId = ContainerHeader->OptionalSegmentPackageIds[Index];
-			check(PackageId.IsValid());
-			OptionalSegmentStoreEntriesMap.FindOrAdd(PackageId, &OptionalSegmentStoreEntry);
-			++Index;
-		}
+			const FIoContainerHeader* ContainerHeader = MountedContainer.ContainerHeader;
+			TArrayView<const FFilePackageStoreEntry> ContainerStoreEntries(reinterpret_cast<const FFilePackageStoreEntry*>(ContainerHeader->StoreEntries.GetData()), ContainerHeader->PackageIds.Num());
+			int32 Index = 0;
+			for (const FFilePackageStoreEntry& StoreEntry : ContainerStoreEntries)
+			{
+				const FPackageId& PackageId = ContainerHeader->PackageIds[Index];
+				check(PackageId.IsValid());
+				StoreEntriesMap.FindOrAdd(PackageId, &StoreEntry);
+				++Index;
+			}
+
+#if WITH_EDITOR
+			TArrayView<const FFilePackageStoreEntry> ContainerOptionalSegmentStoreEntries(reinterpret_cast<const FFilePackageStoreEntry*>(ContainerHeader->OptionalSegmentStoreEntries.GetData()), ContainerHeader->OptionalSegmentPackageIds.Num());
+			Index = 0;
+			for (const FFilePackageStoreEntry& OptionalSegmentStoreEntry : ContainerOptionalSegmentStoreEntries)
+			{
+				const FPackageId& PackageId = ContainerHeader->OptionalSegmentPackageIds[Index];
+				check(PackageId.IsValid());
+				OptionalSegmentStoreEntriesMap.FindOrAdd(PackageId, &OptionalSegmentStoreEntry);
+				++Index;
+			}
 #endif //if WITH_EDITOR
 
-		for (const FIoContainerHeaderLocalizedPackage& LocalizedPackage : ContainerHeader->LocalizedPackages)
-		{
-			FName& LocalizedPackageSourceName = LocalizedPackages.FindOrAdd(LocalizedPackage.SourcePackageId);
-			if (LocalizedPackageSourceName.IsNone())
+			for (const FIoContainerHeaderLocalizedPackage& LocalizedPackage : ContainerHeader->LocalizedPackages)
 			{
-				FDisplayNameEntryId NameEntry = ContainerHeader->RedirectsNameMap[LocalizedPackage.SourcePackageName.GetIndex()];
-				LocalizedPackageSourceName = NameEntry.ToName(LocalizedPackage.SourcePackageName.GetNumber());
+				FName& LocalizedPackageSourceName = LocalizedPackages.FindOrAdd(LocalizedPackage.SourcePackageId);
+				if (LocalizedPackageSourceName.IsNone())
+				{
+					FDisplayNameEntryId NameEntry = ContainerHeader->RedirectsNameMap[LocalizedPackage.SourcePackageName.GetIndex()];
+					LocalizedPackageSourceName = NameEntry.ToName(LocalizedPackage.SourcePackageName.GetNumber());
+				}
 			}
-		}
 
-		for (const FIoContainerHeaderPackageRedirect& Redirect : ContainerHeader->PackageRedirects)
-		{
-			TTuple<FName, FPackageId>& RedirectEntry = RedirectsPackageMap.FindOrAdd(Redirect.SourcePackageId);
-			FName& SourcePackageName = RedirectEntry.Key;
-			if (SourcePackageName.IsNone())
+			for (const FIoContainerHeaderPackageRedirect& Redirect : ContainerHeader->PackageRedirects)
 			{
-				FDisplayNameEntryId NameEntry = ContainerHeader->RedirectsNameMap[Redirect.SourcePackageName.GetIndex()];
-				SourcePackageName = NameEntry.ToName(Redirect.SourcePackageName.GetNumber());
-				RedirectEntry.Value = Redirect.TargetPackageId;
+				TTuple<FName, FPackageId>& RedirectEntry = RedirectsPackageMap.FindOrAdd(Redirect.SourcePackageId);
+				FName& SourcePackageName = RedirectEntry.Key;
+				if (SourcePackageName.IsNone())
+				{
+					FDisplayNameEntryId NameEntry = ContainerHeader->RedirectsNameMap[Redirect.SourcePackageName.GetIndex()];
+					SourcePackageName = NameEntry.ToName(Redirect.SourcePackageName.GetNumber());
+					RedirectEntry.Value = Redirect.TargetPackageId;
+				}
 			}
 		}
+		bNeedsContainerUpdate = false;
 	}
 
 #if WITH_EDITOR
+	if (bNeedsUncookedPackagesUpdate)
 	{
+		TRACE_CPUPROFILER_EVENT_SCOPE(UpdateUncookedPackages);
 		FScopeLock _(&UncookedPackageRootsLock);
 
 		if (!PendingRemoveUncookedPackageRoots.IsEmpty())
@@ -278,10 +283,9 @@ void FFilePackageStoreBackend::Update()
 			PendingAddUncookedPackageRoots.Empty();
 			UE_LOG(LogFilePackageStore, Display, TEXT("Added %lld uncooked packages"), TotalAddedCount);
 		}
+		bNeedsUncookedPackagesUpdate = false;
 	}
 #endif //if WITH_EDITOR
-
-	bNeedsUpdate = false;
 }
 
 #if WITH_EDITOR
diff --git a/Engine/Source/Runtime/PakFile/Private/FilePackageStore.h b/Engine/Source/Runtime/PakFile/Private/FilePackageStore.h
--- a/Engine/Source/Runtime/PakFile/Private/FilePackageStore.h
+++ b/Engine/Source/Runtime/PakFile/Private/FilePackageStore.h
@@ -61,7 +61,7 @@ private:
 	TMap<FPackageId, const FFilePackageStoreEntry*> StoreEntriesMap;
 	TMap<FPackageId, TTuple<FName, FPackageId>> RedirectsPackageMap;
 	TMap<FPackageId, FName> LocalizedPackages;
-	bool bNeedsUpdate = false;
+	bool bNeedsContainerUpdate = false;
 
 #if WITH_EDITOR
 	FDelegateHandle OnContentPathMountedDelegateHandle;
@@ -71,6 +71,7 @@ private:
 	TSet<FString> PendingRemoveUncookedPackageRoots;
 	TMap<FPackageId, FUncookedPackage> UncookedPackagesMap;
 	TMap<FPackageId, const FFilePackageStoreEntry*> OptionalSegmentStoreEntriesMap;
+	bool bNeedsUncookedPackagesUpdate = false;
 #endif //if WITH_EDITOR
 
 };
