IoDispatcher: Fixed various issues related to unmounting a container when there are still requests in flight for that container

diff --git a/Engine/Source/Runtime/PakFile/Internal/IoDispatcherFileBackendTypes.h b/Engine/Source/Runtime/PakFile/Internal/IoDispatcherFileBackendTypes.h
--- a/Engine/Source/Runtime/PakFile/Internal/IoDispatcherFileBackendTypes.h
+++ b/Engine/Source/Runtime/PakFile/Internal/IoDispatcherFileBackendTypes.h
@@ -25,10 +25,10 @@ struct FFileIoStoreCompressionContext;
 struct PAKFILE_API FFileIoStoreContainerFilePartition
 {
 	FFileIoStoreContainerFilePartition() = default;
-	FFileIoStoreContainerFilePartition(FFileIoStoreContainerFilePartition&&) = default;
+	FFileIoStoreContainerFilePartition(FFileIoStoreContainerFilePartition&&) = delete;
 	FFileIoStoreContainerFilePartition(const FFileIoStoreContainerFilePartition&) = delete;
 
-	FFileIoStoreContainerFilePartition& operator=(FFileIoStoreContainerFilePartition&&) = default;
+	FFileIoStoreContainerFilePartition& operator=(FFileIoStoreContainerFilePartition&&) = delete;
 	FFileIoStoreContainerFilePartition& operator=(const FFileIoStoreContainerFilePartition&) = delete;
 
 	uint64 FileHandle = 0;
@@ -36,6 +36,28 @@ struct PAKFILE_API FFileIoStoreContainerFilePartition
 	uint32 ContainerFileIndex = 0;
 	FString FilePath;
 	TUniquePtr<IMappedFileHandle> MappedFileHandle;
+	std::atomic<int32> StartedReadRequestsCount = 0;
+};
+
+struct PAKFILE_API FFileIoStoreBlockSignatureTable
+{
+	TArray<FSHAHash> Hashes;
+
+	void AddRef()
+	{
+		++RefCount;
+	}
+
+	void Release()
+	{
+		if (--RefCount == 0)
+		{
+			delete this;
+		}
+	}
+
+private:
+	std::atomic<int64> RefCount = 0;
 };
 
 struct PAKFILE_API FFileIoStoreContainerFile
@@ -55,15 +77,14 @@ struct PAKFILE_API FFileIoStoreContainerFile
 	FGuid EncryptionKeyGuid;
 	FAES::FAESKey EncryptionKey;
 	EIoContainerFlags ContainerFlags;
-	TArray<FSHAHash> BlockSignatureHashes;
+	TRefCountPtr<FFileIoStoreBlockSignatureTable> BlockSignatureTable;
 	TArray<FFileIoStoreContainerFilePartition> Partitions;
 	uint32 ContainerInstanceId = 0;
 
-	void GetPartitionFileHandleAndOffset(uint64 TocOffset, uint64& OutFileHandle, uint64& OutOffset) const
+	void GetPartitionAndOffset(uint64 TocOffset, FFileIoStoreContainerFilePartition*& OutPartition, uint64& OutOffset)
 	{
 		int32 PartitionIndex = int32(TocOffset / PartitionSize);
-		const FFileIoStoreContainerFilePartition& Partition = Partitions[PartitionIndex];
-		OutFileHandle = Partition.FileHandle;
+		OutPartition = &Partitions[PartitionIndex];
 		OutOffset = TocOffset % PartitionSize;
 	}
 };
@@ -122,6 +143,7 @@ struct PAKFILE_API FFileIoStoreCompressedBlock
 	FFileIoStoreCompressionContext* CompressionContext = nullptr;
 	uint8* CompressedDataBuffer = nullptr;
 	FAES::FAESKey EncryptionKey;
+	TRefCountPtr<FFileIoStoreBlockSignatureTable> BlockSignatureTable;
 	const FSHAHash* SignatureHash = nullptr;
 	bool bFailed = false;
 	bool bCancelled = false;
@@ -133,7 +155,8 @@ struct PAKFILE_API FFileIoStoreReadRequest
 	{
 		QueueStatus_NotInQueue,
 		QueueStatus_InQueue,
-		QueueStatus_Started
+		QueueStatus_Started,
+		QueueStatus_Completed
 	};
 
 	FFileIoStoreReadRequest()
@@ -143,7 +166,7 @@ struct PAKFILE_API FFileIoStoreReadRequest
 	}
 	FFileIoStoreReadRequest* Next = nullptr;
 	FFileIoStoreReadRequest* Previous = nullptr;
-	uint64 FileHandle = uint64(-1);
+	FFileIoStoreContainerFilePartition* ContainerFilePartition = nullptr;
 	uint64 Offset = uint64(-1);
 	uint64 Size = uint64(-1);
 	FFileIoStoreBlockKey Key;
@@ -593,7 +616,7 @@ struct PAKFILE_API FFileIoStoreReadRequestSortKey
 
 	FFileIoStoreReadRequestSortKey() {}
 	FFileIoStoreReadRequestSortKey(FFileIoStoreReadRequest* Request)
-		: Offset(Request->Offset), Handle(Request->FileHandle), Priority(Request->Priority)
+		: Offset(Request->Offset), Handle(Request->ContainerFilePartition->FileHandle), Priority(Request->Priority)
 	{
 	}
 };
@@ -618,7 +641,7 @@ public:
 
 	FFileIoStoreReadRequest* Pop(FFileIoStoreReadRequestSortKey LastSortKey);
 	void Push(FFileIoStoreReadRequest* Request);
-	void CancelRequestsWithFileHandle(uint64 FileHandle);
+	int32 HandleContainerUnmounted(const FFileIoStoreContainerFile& ContainerFile);
 
 private:
 	int32 Priority;
@@ -646,7 +669,7 @@ public:
 	void UpdateOrder();
 	void Lock();
 	void Unlock();
-	void CancelRequestsWithFileHandle(uint64 FileHandle);
+	int32 HandleContainerUnmounted(const FFileIoStoreContainerFile& ContainerFile);
 
 private:
 	static bool QueueSortFunc(const FFileIoStoreReadRequest& A, const FFileIoStoreReadRequest& B)
@@ -771,7 +794,7 @@ public:
 
 	FFileIoStoreResolvedRequest* AllocResolvedRequest(
 		FIoRequestImpl& InDispatcherRequest,
-		const FFileIoStoreContainerFile& InContainerFile,
+		FFileIoStoreContainerFile* InContainerFile,
 		uint64 InResolvedOffset,
 		uint64 InResolvedSize)
 	{
@@ -837,11 +860,16 @@ struct PAKFILE_API FFileIoStoreResolvedRequest
 public:
 	FFileIoStoreResolvedRequest(
 		FIoRequestImpl& InDispatcherRequest,
-		const FFileIoStoreContainerFile& InContainerFile,
+		FFileIoStoreContainerFile* InContainerFile,
 		uint64 InResolvedOffset,
 		uint64 InResolvedSize);
 
-	const FFileIoStoreContainerFile& GetContainerFile() const
+	const FFileIoStoreContainerFile* GetContainerFile() const
+	{
+		return ContainerFile;
+	}
+
+	FFileIoStoreContainerFile* GetContainerFile()
 	{
 		return ContainerFile;
 	}
@@ -884,7 +912,7 @@ public:
 
 private:
 	FIoRequestImpl* DispatcherRequest = nullptr;
-	const FFileIoStoreContainerFile& ContainerFile;
+	FFileIoStoreContainerFile* ContainerFile;
 	FFileIoStoreReadRequestLink* ReadRequestsHead = nullptr;
 	FFileIoStoreReadRequestLink* ReadRequestsTail = nullptr;
 	const uint64 ResolvedOffset;
@@ -949,7 +977,7 @@ private:
 #endif
 
 	// Used for seek tracking 
-	uint64 LastHandle = uint64(-1);
+	const FFileIoStoreContainerFilePartition* LastContainerFilePartition = nullptr;
 	uint64 LastOffset = uint64(-1);
 
 	FTSTicker::FDelegateHandle TickerHandle;
diff --git a/Engine/Source/Runtime/PakFile/Private/GenericPlatformIoDispatcher.cpp b/Engine/Source/Runtime/PakFile/Private/GenericPlatformIoDispatcher.cpp
--- a/Engine/Source/Runtime/PakFile/Private/GenericPlatformIoDispatcher.cpp
+++ b/Engine/Source/Runtime/PakFile/Private/GenericPlatformIoDispatcher.cpp
@@ -75,13 +75,22 @@ void FGenericFileIoStoreImpl::CloseContainer(uint64 ContainerFileHandle)
 
 bool FGenericFileIoStoreImpl::StartRequests(FFileIoStoreRequestQueue& RequestQueue)
 {
+	if (!AcquiredBuffer)
+	{
+		AcquiredBuffer = BufferAllocator->AllocBuffer();
+		if (!AcquiredBuffer)
+		{
+			return false;
+		}
+	}
+
 	FFileIoStoreReadRequest* NextRequest = RequestQueue.Pop();
 	if (!NextRequest)
 	{
 		return false;
 	}
 
-	if (NextRequest->bCancelled)
+	if (NextRequest->bCancelled | NextRequest->bFailed)
 	{
 		{
 			FScopeLock _(&CompletedRequestsCritical);
@@ -93,17 +102,14 @@ bool FGenericFileIoStoreImpl::StartRequests(FFileIoStoreRequestQueue& RequestQue
 
 	uint8* Dest;
 	check(!NextRequest->ImmediateScatter.Request);
-	NextRequest->Buffer = BufferAllocator->AllocBuffer();
-	if (!NextRequest->Buffer)
-	{
-		RequestQueue.Push(*NextRequest);
-		return false;
-	}
+	
+	NextRequest->Buffer = AcquiredBuffer;
+	AcquiredBuffer = nullptr;
 	Dest = NextRequest->Buffer->Memory;
 
 	if (!BlockCache->Read(NextRequest))
 	{
-		IFileHandle* FileHandle = reinterpret_cast<IFileHandle*>(static_cast<UPTRINT>(NextRequest->FileHandle));
+		IFileHandle* FileHandle = reinterpret_cast<IFileHandle*>(static_cast<UPTRINT>(NextRequest->ContainerFilePartition->FileHandle));
 		 
 		Stats->OnFilesystemReadStarted(NextRequest);
 		{
diff --git a/Engine/Source/Runtime/PakFile/Private/GenericPlatformIoDispatcher.h b/Engine/Source/Runtime/PakFile/Private/GenericPlatformIoDispatcher.h
--- a/Engine/Source/Runtime/PakFile/Private/GenericPlatformIoDispatcher.h
+++ b/Engine/Source/Runtime/PakFile/Private/GenericPlatformIoDispatcher.h
@@ -58,6 +58,7 @@ private:
 	FFileIoStoreBufferAllocator* BufferAllocator = nullptr;
 	FFileIoStoreBlockCache* BlockCache = nullptr;
 	FFileIoStoreStats* Stats = nullptr;
+	FFileIoStoreBuffer* AcquiredBuffer = nullptr;
 
 	FCriticalSection CompletedRequestsCritical;
 	FFileIoStoreReadRequestList CompletedRequests;
diff --git a/Engine/Source/Runtime/PakFile/Private/IoDispatcherFileBackend.cpp b/Engine/Source/Runtime/PakFile/Private/IoDispatcherFileBackend.cpp
--- a/Engine/Source/Runtime/PakFile/Private/IoDispatcherFileBackend.cpp
+++ b/Engine/Source/Runtime/PakFile/Private/IoDispatcherFileBackend.cpp
@@ -357,7 +357,7 @@ FFileIoStoreReadRequest* FFileIoStoreOffsetSortedRequestQueue::GetNextInternal(F
 			RequestIndex = Algo::LowerBoundBy(Requests, LastSortKey, RequestSortProjection, RequestSortPredicate);
 			if (Requests.IsValidIndex(RequestIndex)) // If all requests are before LastOffset we get back out-of-bounds
 			{
-				if (Requests[RequestIndex]->FileHandle != LastSortKey.Handle)
+				if (Requests[RequestIndex]->ContainerFilePartition->FileHandle != LastSortKey.Handle)
 				{
 					// Changing file handle so switch back to the oldest outstanding request 
 					RequestIndex = INDEX_NONE;
@@ -417,15 +417,31 @@ void FFileIoStoreOffsetSortedRequestQueue::Push(FFileIoStoreReadRequest* Request
 	PeekRequestIndex = INDEX_NONE;
 }
 
-void FFileIoStoreOffsetSortedRequestQueue::CancelRequestsWithFileHandle(uint64 FileHandle)
+int32 HandleContainerUnmounted(const TArrayView<FFileIoStoreReadRequest*> Requests, const FFileIoStoreContainerFile& ContainerFile)
 {
-	for (FFileIoStoreReadRequest* Request : Requests)
+	static FFileIoStoreContainerFilePartition UnmountedPartition;
+
+	int32 FailedRequestsCount = 0;
+
+	for (const FFileIoStoreContainerFilePartition& Partition : ContainerFile.Partitions)
 	{
-		if (Request->FileHandle == FileHandle)
+		for (FFileIoStoreReadRequest* Request : Requests)
 		{
-			Request->bCancelled = true;
+			if (Request->ContainerFilePartition == &Partition)
+			{
+				Request->bFailed = true;
+				Request->ContainerFilePartition = &UnmountedPartition;
+				++FailedRequestsCount;
+			}
 		}
 	}
+
+	return FailedRequestsCount;
+}
+
+int32 FFileIoStoreOffsetSortedRequestQueue::HandleContainerUnmounted(const FFileIoStoreContainerFile& ContainerFile)
+{
+	return ::HandleContainerUnmounted(Requests, ContainerFile);
 }
 
 void FFileIoStoreRequestQueue::UpdateSortRequestsByOffset()
@@ -502,6 +518,7 @@ FFileIoStoreReadRequest* FFileIoStoreRequestQueue::Pop()
 	
 	check(Result->QueueStatus == FFileIoStoreReadRequest::QueueStatus_InQueue);
 	Result->QueueStatus = FFileIoStoreReadRequest::QueueStatus_Started;
+	Result->ContainerFilePartition->StartedReadRequestsCount.fetch_add(1, std::memory_order_release);
 	return Result;
 }
 
@@ -527,7 +544,7 @@ void FFileIoStoreRequestQueue::Push(FFileIoStoreReadRequest& Request)
 	FScopeLock _(&CriticalSection);
 	UpdateSortRequestsByOffset();
 	
-	check(Request.QueueStatus != FFileIoStoreReadRequest::QueueStatus_InQueue);
+	check(Request.QueueStatus == FFileIoStoreReadRequest::QueueStatus_NotInQueue);
 	Request.QueueStatus = FFileIoStoreReadRequest::QueueStatus_InQueue;
 
 	if (bSortRequestsByOffset)
@@ -548,7 +565,7 @@ void FFileIoStoreRequestQueue::Push(FFileIoStoreReadRequestList& Requests)
 
 	for (auto It = Requests.Steal(); It; ++It)
 	{
-		check(It->QueueStatus != FFileIoStoreReadRequest::QueueStatus_InQueue);
+		check(It->QueueStatus == FFileIoStoreReadRequest::QueueStatus_NotInQueue);
 		It->QueueStatus = FFileIoStoreReadRequest::QueueStatus_InQueue;
 
 		if (bSortRequestsByOffset)
@@ -601,27 +618,25 @@ void FFileIoStoreRequestQueue::Unlock()
 	CriticalSection.Unlock();
 }
 
-void FFileIoStoreRequestQueue::CancelRequestsWithFileHandle(const uint64 FileHandle)
+int32 FFileIoStoreRequestQueue::HandleContainerUnmounted(const FFileIoStoreContainerFile& ContainerFile)
 {
 	FScopeLock _(&CriticalSection);
 	
+	int32 FailedRequestsCount = 0;
+
 	if (bSortRequestsByOffset)
 	{
 		for (FFileIoStoreOffsetSortedRequestQueue& SubQueue : SortedPriorityQueues)
 		{
-			SubQueue.CancelRequestsWithFileHandle(FileHandle);
+			FailedRequestsCount += SubQueue.HandleContainerUnmounted(ContainerFile);
 		}
 	}
 	else
 	{
-		for (FFileIoStoreReadRequest* Request : Heap)
-		{
-			if (Request->FileHandle == FileHandle)
-			{
-				Request->bCancelled = true;
-			}
-		}
+		FailedRequestsCount += ::HandleContainerUnmounted(Heap, ContainerFile);
 	}
+
+	return FailedRequestsCount;
 }
 
 FFileIoStoreReader::FFileIoStoreReader(IPlatformFileIoStore& InPlatformImpl, FFileIoStoreStats& InStats)
@@ -642,7 +657,7 @@ uint64 FFileIoStoreReader::GetTocAllocatedSize() const
 		PerfectHashMap.TocChunkIds.GetAllocatedSize() +
 		PerfectHashMap.TocChunkHashSeeds.GetAllocatedSize() +
 		ContainerFile.CompressionBlocks.GetAllocatedSize() +
-		ContainerFile.BlockSignatureHashes.GetAllocatedSize();
+		(ContainerFile.BlockSignatureTable.IsValid() ? ContainerFile.BlockSignatureTable->Hashes.GetAllocatedSize() : 0);
 }
 
 FIoStatus FFileIoStoreReader::Initialize(const TCHAR* InTocFilePath, int32 InOrder)
@@ -714,7 +729,8 @@ FIoStatus FFileIoStoreReader::Initialize(const TCHAR* InTocFilePath, int32 InOrd
 	ContainerFile.CompressionBlocks		= MoveTemp(TocResource.CompressionBlocks);
 	ContainerFile.ContainerFlags		= TocResource.Header.ContainerFlags;
 	ContainerFile.EncryptionKeyGuid		= TocResource.Header.EncryptionKeyGuid;
-	ContainerFile.BlockSignatureHashes	= MoveTemp(TocResource.ChunkBlockSignatures);
+	ContainerFile.BlockSignatureTable	= new FFileIoStoreBlockSignatureTable();
+	ContainerFile.BlockSignatureTable->Hashes = MoveTemp(TocResource.ChunkBlockSignatures);
 	ContainerFile.ContainerInstanceId	= ++GlobalContainerInstanceId;
 
 	Stats.OnTocMounted(GetTocAllocatedSize());
@@ -903,7 +919,7 @@ TIoStatusOr<FIoContainerHeader> FFileIoStoreReader::ReadContainerHeader() const
 			const uint64 BlockSize = Align(CompressionBlockEntry->GetCompressedSize(), FAES::AESBlockSize);
 			if (bSigned)
 			{
-				const FSHAHash& SignatureHash = ContainerFile.BlockSignatureHashes[CompressedBlockIndex];
+				const FSHAHash& SignatureHash = ContainerFile.BlockSignatureTable->Hashes[CompressedBlockIndex];
 				FSHAHash BlockHash;
 				FSHA1::HashBuffer(BlockData, BlockSize, BlockHash.Hash);
 				if (SignatureHash != BlockHash)
@@ -940,7 +956,7 @@ void FFileIoStoreReader::ReopenAllFileHandles()
 
 FFileIoStoreResolvedRequest::FFileIoStoreResolvedRequest(
 	FIoRequestImpl& InDispatcherRequest,
-	const FFileIoStoreContainerFile& InContainerFile,
+	FFileIoStoreContainerFile* InContainerFile,
 	uint64 InResolvedOffset,
 	uint64 InResolvedSize)
 	: DispatcherRequest(&InDispatcherRequest)
@@ -1080,7 +1096,7 @@ bool FFileIoStoreRequestTracker::CancelIoRequest(FFileIoStoreResolvedRequest& Re
 			continue;
 		}
 
-		if (ReadRequest.QueueStatus == FFileIoStoreReadRequest::QueueStatus_Started)
+		if (ReadRequest.QueueStatus >= FFileIoStoreReadRequest::QueueStatus_Started)
 		{
 			bShouldComplete = false;
 			continue;
@@ -1291,28 +1307,47 @@ TIoStatusOr<FIoContainerHeader> FFileIoStore::Mount(const TCHAR* InTocPath, int3
 
 bool FFileIoStore::Unmount(const TCHAR* InTocPath)
 {
-	FWriteScopeLock _(IoStoreReadersLock);
-	
-	for (int32 Idx = 0; Idx < IoStoreReaders.Num(); ++Idx)
+	TUniquePtr<FFileIoStoreReader> ReaderToUnmount;
 	{
-		if (IoStoreReaders[Idx]->GetContainerFile().FilePath == InTocPath)
-		{
-			UE_LOG(LogIoDispatcher, Display, TEXT("Unmounting container '%s'"), InTocPath);
+		FWriteScopeLock _(IoStoreReadersLock);
 
-			// Cancel pending I/O requests trying to read from the container
-			for (const FFileIoStoreContainerFilePartition& Partition : IoStoreReaders[Idx]->GetContainerFile().Partitions)
+		for (int32 Idx = 0; Idx < IoStoreReaders.Num(); ++Idx)
+		{
+			if (IoStoreReaders[Idx]->GetContainerFile()->FilePath == InTocPath)
 			{
-				RequestQueue.CancelRequestsWithFileHandle(Partition.FileHandle);
-			}
+				ReaderToUnmount = MoveTemp(IoStoreReaders[Idx]);
+				IoStoreReaders.RemoveAt(Idx);
 
-			FIoContainerId ContainerId = IoStoreReaders[Idx]->GetContainerId();
-			IoStoreReaders.RemoveAt(Idx);
-			
-			return true;
+				break;
+			}
 		}
 	}
+	if (ReaderToUnmount)
+	{
+		UE_LOG(LogIoDispatcher, Display, TEXT("Unmounting container '%s'"), InTocPath);
 
-	UE_LOG(LogIoDispatcher, Display, TEXT("Failed to unmount container '%s'"), InTocPath);
+		int32 FailedRequestsCount = RequestQueue.HandleContainerUnmounted(*ReaderToUnmount->GetContainerFile());
+
+		UE_CLOG(FailedRequestsCount > 0, LogIoDispatcher, Warning, TEXT("Marking %d queued requests from unmounted container as failed"), FailedRequestsCount);
+
+		bool bHasWarned = false;
+		for (const FFileIoStoreContainerFilePartition& Partition : ReaderToUnmount->GetContainerFile()->Partitions)
+		{
+			if (Partition.StartedReadRequestsCount.load(std::memory_order_acquire) != 0)
+			{
+				UE_CLOG(!bHasWarned, LogIoDispatcher, Warning, TEXT("Waiting for read requests to finish before unmounting container"));
+				bHasWarned = true;
+				while (Partition.StartedReadRequestsCount.load(std::memory_order_acquire) != 0)
+				{
+					FPlatformProcess::Sleep(0);
+				}
+			}
+		}
+	}
+	else
+	{
+		UE_LOG(LogIoDispatcher, Display, TEXT("Failed to unmount container '%s'"), InTocPath);
+	}
 	
 	return false;
 }
@@ -1453,38 +1488,39 @@ void FFileIoStore::ScatterBlock(FFileIoStoreCompressedBlock* CompressedBlock, bo
 		uint64 OffsetInBuffer = CompressedBlock->RawOffset - RawBlock->Offset;
 		CompressedBuffer = RawBlock->Buffer->Memory + OffsetInBuffer;
 	}
-	if (CompressedBlock->SignatureHash)
+	if (!CompressedBlock->bFailed)
 	{
-		FSHAHash BlockHash;
-		FSHA1::HashBuffer(CompressedBuffer, CompressedBlock->RawSize, BlockHash.Hash);
-		if (*CompressedBlock->SignatureHash != BlockHash)
+		if (CompressedBlock->SignatureHash)
 		{
-			FIoSignatureError Error;
+			FSHAHash BlockHash;
+			FSHA1::HashBuffer(CompressedBuffer, CompressedBlock->RawSize, BlockHash.Hash);
+			if (*CompressedBlock->SignatureHash != BlockHash)
 			{
-				FReadScopeLock _(IoStoreReadersLock);
-				for (const TUniquePtr<FFileIoStoreReader>& Reader : IoStoreReaders)
+				FIoSignatureError Error;
 				{
-					if (CompressedBlock->Key.FileIndex == Reader->GetContainerInstanceId())
+					FReadScopeLock _(IoStoreReadersLock);
+					for (const TUniquePtr<FFileIoStoreReader>& Reader : IoStoreReaders)
 					{
-						Error.ContainerName = FPaths::GetBaseFilename(Reader->GetContainerFile().FilePath);
+						if (CompressedBlock->Key.FileIndex == Reader->GetContainerInstanceId())
+						{
+							Error.ContainerName = FPaths::GetBaseFilename(Reader->GetContainerFile()->FilePath);
+						}
 					}
+					Error.BlockIndex = CompressedBlock->Key.BlockIndex;
+					Error.ExpectedHash = *CompressedBlock->SignatureHash;
+					Error.ActualHash = BlockHash;
 				}
-				Error.BlockIndex = CompressedBlock->Key.BlockIndex;
-				Error.ExpectedHash = *CompressedBlock->SignatureHash;
-				Error.ActualHash = BlockHash;
-			}
 
-			UE_LOG(LogIoDispatcher, Warning, TEXT("Signature error detected in container '%s' at block index '%d'"), *Error.ContainerName, Error.BlockIndex);
+				UE_LOG(LogIoDispatcher, Warning, TEXT("Signature error detected in container '%s' at block index '%d'"), *Error.ContainerName, Error.BlockIndex);
 
-			check(BackendContext);
-			if (BackendContext->SignatureErrorDelegate.IsBound())
-			{
-				BackendContext->SignatureErrorDelegate.Broadcast(Error);
+				check(BackendContext);
+				if (BackendContext->SignatureErrorDelegate.IsBound())
+				{
+					BackendContext->SignatureErrorDelegate.Broadcast(Error);
+				}
 			}
 		}
-	}
-	if (!CompressedBlock->bFailed)
-	{
+
 		if (CompressedBlock->EncryptionKey.IsValid())
 		{
 			FAES::DecryptData(CompressedBuffer, CompressedBlock->RawSize, CompressedBlock->EncryptionKey);
@@ -1574,7 +1610,7 @@ void FFileIoStore::FinalizeCompressedBlock(FFileIoStoreCompressedBlock* Compress
 		check(RawBlock->BufferRefCount > 0);
 		if (--RawBlock->BufferRefCount == 0)
 		{
-			check(RawBlock->Buffer || RawBlock->bCancelled);
+			check(RawBlock->Buffer || RawBlock->bCancelled || RawBlock->bFailed);
 			if (RawBlock->Buffer)
 			{
 				FreeBuffer(*RawBlock->Buffer);
@@ -1622,9 +1658,14 @@ FIoRequestImpl* FFileIoStore::GetCompletedRequests()
 	{
 		FFileIoStoreReadRequest* CompletedRequest = *It;
 
+		check(CompletedRequest->QueueStatus == FFileIoStoreReadRequest::QueueStatus_Started);
+		CompletedRequest->QueueStatus = FFileIoStoreReadRequest::QueueStatus_Completed;
+		int32 PreviousStartedReadRequestsCount = CompletedRequest->ContainerFilePartition->StartedReadRequestsCount.fetch_sub(1, std::memory_order_release);
+		check(PreviousStartedReadRequestsCount >= 1);
+
 		if (!CompletedRequest->ImmediateScatter.Request)
 		{
-			check(CompletedRequest->Buffer || CompletedRequest->bCancelled);
+			check(CompletedRequest->Buffer || CompletedRequest->bCancelled || CompletedRequest->bFailed);
 			RequestTracker.RemoveRawBlock(CompletedRequest);
 			
 			//TRACE_CPUPROFILER_EVENT_SCOPE(ProcessCompletedBlock);
@@ -1802,10 +1843,10 @@ TIoStatusOr<FIoMappedRegion> FFileIoStore::OpenMapped(const FIoChunkId& ChunkId,
 			uint64 ResolvedOffset = OffsetAndLength->GetOffset();
 			uint64 ResolvedSize = FMath::Min(Options.GetSize(), OffsetAndLength->GetLength());
 			
-			const FFileIoStoreContainerFile& ContainerFile = Reader->GetContainerFile();
+			const FFileIoStoreContainerFile* ContainerFile = Reader->GetContainerFile();
 			
-			int32 BlockIndex = int32(ResolvedOffset / ContainerFile.CompressionBlockSize);
-			const FIoStoreTocCompressedBlockEntry& CompressionBlockEntry = ContainerFile.CompressionBlocks[BlockIndex];
+			int32 BlockIndex = int32(ResolvedOffset / ContainerFile->CompressionBlockSize);
+			const FIoStoreTocCompressedBlockEntry& CompressionBlockEntry = ContainerFile->CompressionBlocks[BlockIndex];
 			const int64 BlockOffset = (int64)CompressionBlockEntry.GetOffset();
 			check(BlockOffset > 0 && IsAligned(BlockOffset, FPlatformProperties::GetMemoryMappingAlignment()));
 
@@ -1851,8 +1892,8 @@ void FFileIoStore::ReadBlocks(FFileIoStoreResolvedRequest& ResolvedRequest)
 	ScopeName.Appendf(TEXT("ReadBlock %d"), BlockIndex);
 	TRACE_CPUPROFILER_EVENT_SCOPE_TEXT(*ScopeName);*/
 
-	const FFileIoStoreContainerFile& ContainerFile = ResolvedRequest.GetContainerFile();
-	const uint64 CompressionBlockSize = ContainerFile.CompressionBlockSize;
+	FFileIoStoreContainerFile* ContainerFile = ResolvedRequest.GetContainerFile();
+	const uint64 CompressionBlockSize = ContainerFile->CompressionBlockSize;
 	const uint64 RequestEndOffset = ResolvedRequest.ResolvedOffset + ResolvedRequest.ResolvedSize;
 	int32 RequestBeginBlockIndex = int32(ResolvedRequest.ResolvedOffset / CompressionBlockSize);
 	int32 RequestEndBlockIndex = int32((RequestEndOffset - 1) / CompressionBlockSize);
@@ -1865,7 +1906,7 @@ void FFileIoStore::ReadBlocks(FFileIoStoreResolvedRequest& ResolvedRequest)
 	for (int32 CompressedBlockIndex = RequestBeginBlockIndex; CompressedBlockIndex <= RequestEndBlockIndex; ++CompressedBlockIndex)
 	{
 		FFileIoStoreBlockKey CompressedBlockKey;
-		CompressedBlockKey.FileIndex = ResolvedRequest.GetContainerFile().ContainerInstanceId;
+		CompressedBlockKey.FileIndex = ContainerFile->ContainerInstanceId;
 		CompressedBlockKey.BlockIndex = CompressedBlockIndex;
 		bool bCompressedBlockWasAdded;
 		FFileIoStoreCompressedBlock* CompressedBlock = RequestTracker.FindOrAddCompressedBlock(CompressedBlockKey, bCompressedBlockWasAdded);
@@ -1873,17 +1914,22 @@ void FFileIoStore::ReadBlocks(FFileIoStoreResolvedRequest& ResolvedRequest)
 		check(!CompressedBlock->bCancelled);
 		if (bCompressedBlockWasAdded)
 		{
-			CompressedBlock->EncryptionKey = ContainerFile.EncryptionKey;
-			const FIoStoreTocCompressedBlockEntry& CompressionBlockEntry = ContainerFile.CompressionBlocks[CompressedBlockIndex];
+			CompressedBlock->EncryptionKey = ContainerFile->EncryptionKey;
+			const FIoStoreTocCompressedBlockEntry& CompressionBlockEntry = ContainerFile->CompressionBlocks[CompressedBlockIndex];
 			CompressedBlock->UncompressedSize = CompressionBlockEntry.GetUncompressedSize();
 			CompressedBlock->CompressedSize = CompressionBlockEntry.GetCompressedSize();
-			CompressedBlock->CompressionMethod = ContainerFile.CompressionMethods[CompressionBlockEntry.GetCompressionMethodIndex()];
-			CompressedBlock->SignatureHash = EnumHasAnyFlags(ContainerFile.ContainerFlags, EIoContainerFlags::Signed) ? &ContainerFile.BlockSignatureHashes[CompressedBlockIndex] : nullptr;
+			CompressedBlock->CompressionMethod = ContainerFile->CompressionMethods[CompressionBlockEntry.GetCompressionMethodIndex()];
+			if (EnumHasAnyFlags(ContainerFile->ContainerFlags, EIoContainerFlags::Signed))
+			{
+				check(ContainerFile->BlockSignatureTable);
+				CompressedBlock->BlockSignatureTable = ContainerFile->BlockSignatureTable;
+				CompressedBlock->SignatureHash = &ContainerFile->BlockSignatureTable->Hashes[CompressedBlockIndex];
+			}
 			CompressedBlock->RawSize = Align(CompressionBlockEntry.GetCompressedSize(), FAES::AESBlockSize); // The raw blocks size is always aligned to AES blocks size;
 
-			int32 PartitionIndex = int32(CompressionBlockEntry.GetOffset() / ContainerFile.PartitionSize);
-			const FFileIoStoreContainerFilePartition& Partition = ContainerFile.Partitions[PartitionIndex];
-			uint64 PartitionRawOffset = CompressionBlockEntry.GetOffset() % ContainerFile.PartitionSize;
+			int32 PartitionIndex = int32(CompressionBlockEntry.GetOffset() / ContainerFile->PartitionSize);
+			FFileIoStoreContainerFilePartition& Partition = ContainerFile->Partitions[PartitionIndex];
+			uint64 PartitionRawOffset = CompressionBlockEntry.GetOffset() % ContainerFile->PartitionSize;
 			CompressedBlock->RawOffset = PartitionRawOffset;
 			const uint32 RawBeginBlockIndex = uint32(PartitionRawOffset / ReadBufferSize);
 			const uint32 RawEndBlockIndex = uint32((PartitionRawOffset + CompressedBlock->RawSize - 1) / ReadBufferSize);
@@ -1902,7 +1948,7 @@ void FFileIoStore::ReadBlocks(FFileIoStoreResolvedRequest& ResolvedRequest)
 				if (bRawBlockWasAdded)
 				{
 					RawBlock->Priority = ResolvedRequest.GetPriority();
-					RawBlock->FileHandle = Partition.FileHandle;
+					RawBlock->ContainerFilePartition = &Partition;
 					RawBlock->Offset = RawBlockIndex * ReadBufferSize;
 					uint64 ReadSize = FMath::Min(Partition.FileSize, RawBlock->Offset + ReadBufferSize) - RawBlock->Offset;
 					RawBlock->Size = ReadSize;
@@ -2141,7 +2187,7 @@ void FFileIoStoreStats::OnFilesystemReadStarted(const FFileIoStoreReadRequest* R
 {
 	CSV_CUSTOM_STAT_DEFINED(FrameFilesystemReads, 1, ECsvCustomStatOp::Accumulate);
 
-	if (LastHandle != Request->FileHandle)
+	if (LastContainerFilePartition != Request->ContainerFilePartition)
 	{
 		OnHandleChangeSeek();
 	}
@@ -2154,7 +2200,7 @@ void FFileIoStoreStats::OnFilesystemReadStarted(const FFileIoStoreReadRequest* R
 		OnSeek(LastOffset, Request->Offset);
 	}
 	LastOffset = Request->Offset + Request->Size;
-	LastHandle = Request->FileHandle;
+	LastContainerFilePartition = Request->ContainerFilePartition;
 
 	CSV_CUSTOM_STAT_DEFINED(FrameFilesystemReads, 1, ECsvCustomStatOp::Accumulate);
 }
@@ -2165,7 +2211,7 @@ void FFileIoStoreStats::OnFilesystemReadsStarted(const FFileIoStoreReadRequestLi
 	int32 NumReads = 0;
 	for (const FFileIoStoreReadRequest* Request : Requests)
 	{
-		if (LastHandle != Request->FileHandle)
+		if (LastContainerFilePartition != Request->ContainerFilePartition)
 		{
 			OnHandleChangeSeek();
 		}
@@ -2179,7 +2225,7 @@ void FFileIoStoreStats::OnFilesystemReadsStarted(const FFileIoStoreReadRequestLi
 		}
 
 		LastOffset = Request->Offset + Request->Size;
-		LastHandle = Request->FileHandle;
+		LastContainerFilePartition = Request->ContainerFilePartition;
 
 		++NumReads;
 		TotalBytes += Request->Size;
diff --git a/Engine/Source/Runtime/PakFile/Private/IoDispatcherFileBackend.h b/Engine/Source/Runtime/PakFile/Private/IoDispatcherFileBackend.h
--- a/Engine/Source/Runtime/PakFile/Private/IoDispatcherFileBackend.h
+++ b/Engine/Source/Runtime/PakFile/Private/IoDispatcherFileBackend.h
@@ -36,7 +36,8 @@ public:
 	bool DoesChunkExist(const FIoChunkId& ChunkId) const;
 	TIoStatusOr<uint64> GetSizeForChunk(const FIoChunkId& ChunkId) const;
 	const FIoOffsetAndLength* Resolve(const FIoChunkId& ChunkId) const;
-	const FFileIoStoreContainerFile& GetContainerFile() const { return ContainerFile; }
+	FFileIoStoreContainerFile* GetContainerFile() { return &ContainerFile; }
+	const FFileIoStoreContainerFile* GetContainerFile() const { return &ContainerFile; }
 	IMappedFileHandle* GetMappedContainerFileHandle(uint64 TocOffset);
 	const FIoContainerId& GetContainerId() const { return ContainerId; }
 	int32 GetOrder() const { return Order; }
