Fix for SignedArchiveReader requesting one chunk too many when doing reads that ends on a 64K chunk boundary

#rnx

diff --git a/Engine/Source/Runtime/PakFile/Private/SignedArchiveReader.cpp b/Engine/Source/Runtime/PakFile/Private/SignedArchiveReader.cpp
--- a/Engine/Source/Runtime/PakFile/Private/SignedArchiveReader.cpp
+++ b/Engine/Source/Runtime/PakFile/Private/SignedArchiveReader.cpp
@@ -373,25 +373,19 @@ int64 FSignedArchiveReader::PrecacheChunks(TArray<FSignedArchiveReader::FReadInf
 	SCOPE_SECONDS_ACCUMULATOR(STAT_SignedArchiveReader_PreCacheChunks);
 
 	// Request all the chunks that are needed to complete this read
-	int64 DataOffset;
+	int64 DataOffset = 0;
 	int64 DestOffset = 0;
-	int32 FirstChunkIndex = CalculateChunkIndex(PakOffset);
-	int64 ChunkStartOffset = CalculateChunkOffset(PakOffset, DataOffset);
-	int64 NumChunksForRequest = (DataOffset - ChunkStartOffset + Length) / FPakInfo::MaxChunkDataSize + 1;
-	int64 NumChunks = NumChunksForRequest;
+	const int32 FirstChunkIndex = CalculateChunkIndex(PakOffset);
+	const int32 LastChunkIndex = CalculateChunkIndex(PakOffset + Length - 1);
+	const int64 NumChunks = LastChunkIndex - FirstChunkIndex + 1;
+	int64 ChunkStartOffset = 0;
 	int64 RemainingLength = Length;
 	int64 ArchiveOffset = PakOffset;
-	
-	// And then try to precache 'PrecacheLength' more chunks because it's likely
-	// we're going to try to read them next
-	if ((NumChunks + FirstChunkIndex + PrecacheLength - 1) < ChunkCount)
-	{
-		NumChunks += PrecacheLength;
-	}
+
 	Chunks.Empty(NumChunks);
-	for (int32 ChunkIndexOffset = 0; ChunkIndexOffset < NumChunks; ++ChunkIndexOffset)
+	for (int32 ChunkIndex = FirstChunkIndex; ChunkIndex <= LastChunkIndex; ++ChunkIndex)
 	{
-		ChunkStartOffset = RemainingLength > 0 ? CalculateChunkOffset(ArchiveOffset, DataOffset) : CalculateChunkOffsetFromIndex(ChunkIndexOffset + FirstChunkIndex);
+		ChunkStartOffset = RemainingLength > 0 ? CalculateChunkOffset(ArchiveOffset, DataOffset) : CalculateChunkOffsetFromIndex(ChunkIndex);
 		int64 SizeToReadFromBuffer = RemainingLength;
 		if (DataOffset + SizeToReadFromBuffer > ChunkStartOffset + FPakInfo::MaxChunkDataSize)
 		{
@@ -403,7 +397,6 @@ int64 FSignedArchiveReader::PrecacheChunks(TArray<FSignedArchiveReader::FReadInf
 		ChunkInfo.DestOffset = DestOffset;
 		ChunkInfo.Size = SizeToReadFromBuffer;
 
-		const int32 ChunkIndex = ChunkIndexOffset + FirstChunkIndex;							
 		if (LastCachedChunk.ChunkIndex == ChunkIndex)
 		{
 			ChunkInfo.Request = NULL;
@@ -424,7 +417,7 @@ int64 FSignedArchiveReader::PrecacheChunks(TArray<FSignedArchiveReader::FReadInf
 		RemainingLength -= SizeToReadFromBuffer;
 	}
 
-	return NumChunksForRequest;
+	return NumChunks;
 }
 
 void FSignedArchiveReader::Serialize(void* Data, int64 Length)
diff --git a/Engine/Source/Runtime/PakFile/Private/SignedArchiveReader.h b/Engine/Source/Runtime/PakFile/Private/SignedArchiveReader.h
--- a/Engine/Source/Runtime/PakFile/Private/SignedArchiveReader.h
+++ b/Engine/Source/Runtime/PakFile/Private/SignedArchiveReader.h
@@ -220,12 +220,6 @@ public:
  */
 class FSignedArchiveReader : public FArchive
 {
-	enum
-	{
-		/** Number of chunks to pre-cache.*/
-		PrecacheLength = 0
-	};
-
 	struct FReadInfo
 	{
 		FChunkRequest* Request;
