Add UE_FORCE_USE_PAKS and UE_FORCE_USE_IOSTORE preprocessors so a target can set them for devs to stop needing to manually pass those arguments.

diff --git a/Engine/Source/Runtime/Core/Private/IO/IoDispatcher.cpp b/Engine/Source/Runtime/Core/Private/IO/IoDispatcher.cpp
--- a/Engine/Source/Runtime/Core/Private/IO/IoDispatcher.cpp
+++ b/Engine/Source/Runtime/Core/Private/IO/IoDispatcher.cpp
@@ -913,7 +913,7 @@ HasScriptObjectsChunk(FIoDispatcher& Dispatcher)
 static bool
 HasUseIoStoreParam()
 {
-	static bool bForceIoStore = WITH_IOSTORE_IN_EDITOR && FParse::Param(FCommandLine::Get(), TEXT("UseIoStore"));
+	static bool bForceIoStore = UE_FORCE_USE_IOSTORE || (WITH_IOSTORE_IN_EDITOR && FParse::Param(FCommandLine::Get(), TEXT("UseIoStore")));
 	return bForceIoStore;
 }
 
diff --git a/Engine/Source/Runtime/Core/Public/Misc/Build.h b/Engine/Source/Runtime/Core/Public/Misc/Build.h
--- a/Engine/Source/Runtime/Core/Public/Misc/Build.h
+++ b/Engine/Source/Runtime/Core/Public/Misc/Build.h
@@ -476,6 +476,13 @@
 #define WITH_IOSTORE_IN_EDITOR UE_IS_COOKED_EDITOR
 #endif
 
+// Controls if iostore will be forced on
+#ifndef UE_FORCE_USE_IOSTORE
+#define UE_FORCE_USE_IOSTORE 0??// Controls if paks will be forced on unless -NoPaks argument is passed
+#ifndef UE_FORCE_USE_PAKS
+#define UE_FORCE_USE_PAKS 0
+#endif
+
 // Controls whether Iris networking code is compiled in or not; should normally be defined by UBT
 #ifndef UE_WITH_IRIS
 	#define UE_WITH_IRIS 0
diff --git a/Engine/Source/Runtime/CoreUObject/Private/Serialization/AsyncPackageLoader.cpp b/Engine/Source/Runtime/CoreUObject/Private/Serialization/AsyncPackageLoader.cpp
--- a/Engine/Source/Runtime/CoreUObject/Private/Serialization/AsyncPackageLoader.cpp
+++ b/Engine/Source/Runtime/CoreUObject/Private/Serialization/AsyncPackageLoader.cpp
@@ -194,7 +194,7 @@ void InitAsyncThread()
 		bCommandLineEnabled = FParse::Param(FCommandLine::Get(), TEXT("ZenLoader"));
 		check(GConfig);
 		GConfig->GetBool(TEXT("/Script/Engine.EditorStreamingSettings"), TEXT("s.ZenLoaderEnabled"), bSettingsEnabled, GEngineIni);
-		bHasUseIoStoreParamInEditor = FParse::Param(FCommandLine::Get(), TEXT("UseIoStore"));
+		bHasUseIoStoreParamInEditor = UE_FORCE_USE_IOSTORE || FParse::Param(FCommandLine::Get(), TEXT("UseIoStore"));
 #endif
 		FIoDispatcher& IoDispatcher = FIoDispatcher::Get();
 		bool bHasScriptObjectsChunk = IoDispatcher.DoesChunkExist(CreateIoChunkId(0, 0, EIoChunkType::ScriptObjects));
diff --git a/Engine/Source/Runtime/PakFile/Private/IPlatformFilePak.cpp b/Engine/Source/Runtime/PakFile/Private/IPlatformFilePak.cpp
--- a/Engine/Source/Runtime/PakFile/Private/IPlatformFilePak.cpp
+++ b/Engine/Source/Runtime/PakFile/Private/IPlatformFilePak.cpp
@@ -7401,14 +7401,19 @@ bool FPakPlatformFile::ShouldBeUsed(IPlatformFile* Inner, const TCHAR* CmdLine)
 #endif //if WITH_EDITOR
 
 	bool Result = false;
-#if (!WITH_EDITOR || IS_MONOLITHIC)
+#if (!WITH_EDITOR || IS_MONOLITHIC || UE_FORCE_USE_PAKS)
 	if (!FParse::Param(CmdLine, TEXT("NoPak")))
 	{
+#if UE_FORCE_USE_PAKS
+		// Target wants pak file to be used even if no pak files currently exist (they can be downloaded later at runtime)
+		Result = true;
+#else
 		TArray<FString> PakFolders;
 		GetPakFolders(CmdLine, PakFolders);
 		Result = CheckIfPakFilesExist(Inner, PakFolders);
+#endif // UE_FORCE_USE_PAKS
 	}
-#endif //if (!WITH_EDITOR || IS_MONOLITHIC)
+#endif //if (!WITH_EDITOR || IS_MONOLITHIC || UE_FORCE_USE_PAKS)
 	return Result;
 }
 
@@ -7453,7 +7458,7 @@ bool FPakPlatformFile::Initialize(IPlatformFile* Inner, const TCHAR* CmdLine)
 
 	FString GlobalUTocPath = FString::Printf(TEXT("%sPaks/global.utoc"), *FPaths::ProjectContentDir());
 	const bool bShouldMountGlobal = FPlatformFileManager::Get().GetPlatformFile().FileExists(*GlobalUTocPath);
-	const bool bForceIoStore = WITH_IOSTORE_IN_EDITOR && FParse::Param(CmdLine, TEXT("UseIoStore"));
+	const bool bForceIoStore = UE_FORCE_USE_IOSTORE || (WITH_IOSTORE_IN_EDITOR && FParse::Param(CmdLine, TEXT("UseIoStore")));
 	if (bShouldMountGlobal || bForceIoStore)
 	{
 		if (ShouldCheckPak())
