Fix Static Analysis Win64 (MSVC): UE_FORCE_USE_IOSTORE
#rnx

diff --git a/Engine/Source/Runtime/PakFile/Private/IPlatformFilePak.cpp b/Engine/Source/Runtime/PakFile/Private/IPlatformFilePak.cpp
--- a/Engine/Source/Runtime/PakFile/Private/IPlatformFilePak.cpp
+++ b/Engine/Source/Runtime/PakFile/Private/IPlatformFilePak.cpp
@@ -7417,6 +7417,17 @@ bool FPakPlatformFile::ShouldBeUsed(IPlatformFile* Inner, const TCHAR* CmdLine)
 	return Result;
 }
 
+static bool PakPlatformFile_IsForceUseIoStore(const TCHAR* CmdLine)
+{
+#if UE_FORCE_USE_IOSTORE
+	return true;
+#elif WITH_IOSTORE_IN_EDITOR
+	return FParse::Param(CmdLine, TEXT("UseIoStore"));
+#else
+	return false;
+#endif
+}
+
 bool FPakPlatformFile::Initialize(IPlatformFile* Inner, const TCHAR* CmdLine)
 {
 	LLM_SCOPE(ELLMTag::FileSystem);
@@ -7458,8 +7469,7 @@ bool FPakPlatformFile::Initialize(IPlatformFile* Inner, const TCHAR* CmdLine)
 
 	FString GlobalUTocPath = FString::Printf(TEXT("%sPaks/global.utoc"), *FPaths::ProjectContentDir());
 	const bool bShouldMountGlobal = FPlatformFileManager::Get().GetPlatformFile().FileExists(*GlobalUTocPath);
-	const bool bForceIoStore = UE_FORCE_USE_IOSTORE || (WITH_IOSTORE_IN_EDITOR && FParse::Param(CmdLine, TEXT("UseIoStore")));
-	if (bShouldMountGlobal || bForceIoStore)
+	if (bShouldMountGlobal || PakPlatformFile_IsForceUseIoStore(CmdLine))
 	{
 		if (ShouldCheckPak())
 		{
