IoDispatcher: Fixed crash when requesting reads of 0 bytes

diff --git a/Engine/Source/Runtime/Core/Private/IO/IoDispatcher.cpp b/Engine/Source/Runtime/Core/Private/IO/IoDispatcher.cpp
--- a/Engine/Source/Runtime/Core/Private/IO/IoDispatcher.cpp
+++ b/Engine/Source/Runtime/Core/Private/IO/IoDispatcher.cpp
@@ -625,6 +625,7 @@ private:
 				}
 				else
 				{
+					UE_CLOG(!CompletedRequestsHead->HasBuffer(), LogStreaming, Fatal, TEXT("Backend provided a completed request without an IoBuiffer. Requests that are not failed or cancelled must have an IoBuffer"));
 					FPlatformAtomics::InterlockedAdd(&TotalLoaded, CompletedRequestsHead->GetBuffer().DataSize());
 					CompleteRequest(CompletedRequestsHead, EIoErrorCode::Ok);
 				}
diff --git a/Engine/Source/Runtime/PakFile/Private/IoDispatcherFileBackend.cpp b/Engine/Source/Runtime/PakFile/Private/IoDispatcherFileBackend.cpp
--- a/Engine/Source/Runtime/PakFile/Private/IoDispatcherFileBackend.cpp
+++ b/Engine/Source/Runtime/PakFile/Private/IoDispatcherFileBackend.cpp
@@ -1392,6 +1392,14 @@ bool FFileIoStore::Resolve(FIoRequestImpl* Request)
 			else
 			{
 				// Nothing to read
+				if (RequestedOffset > OffsetAndLength->GetLength())
+				{
+					ResolvedRequest->bFailed = true;
+				}
+				else
+				{
+					ResolvedRequest->CreateBuffer(0);
+				}
 				CompleteDispatcherRequest(ResolvedRequest);
 				RequestTracker.ReleaseIoRequestReferences(*ResolvedRequest);
 			}
