Fix IPlatformFilePak deadlock when mounting and unmounting simultaneously from different threads.
Refactor Mount and Unmount to only acquire one lock at a time.
Mount was calling into FPakPlatformFile::GetPakSignatureFile() and acquiring PakSignatureFileCacheLock before calling CreateFileReader() which ends up acquiring PakListCritical in FPakPlatformFile::GetMountedPaks().
Unmount was acquiring PakListCritical first and then calling into FPakPlatformFile::RemoveCachedPakSignaturesFile() and acquiring PakSignatureFileCacheLock.

#rnx

diff --git a/Engine/Source/Runtime/PakFile/Private/IPlatformFilePak.cpp b/Engine/Source/Runtime/PakFile/Private/IPlatformFilePak.cpp
--- a/Engine/Source/Runtime/PakFile/Private/IPlatformFilePak.cpp
+++ b/Engine/Source/Runtime/PakFile/Private/IPlatformFilePak.cpp
@@ -351,21 +351,21 @@ void FPakPlatformFile::GetPakEncryptionKey(FAES::FAESKey& OutKey, const FGuid& I
 	}
 }
 
-TMap<FName, TSharedPtr<const struct FPakSignatureFile, ESPMode::ThreadSafe>> FPakPlatformFile::PakSignatureFileCache;
+TMap<FName, TSharedPtr<const FPakSignatureFile, ESPMode::ThreadSafe>> FPakPlatformFile::PakSignatureFileCache;
 FCriticalSection FPakPlatformFile::PakSignatureFileCacheLock;
-TSharedPtr<const struct FPakSignatureFile, ESPMode::ThreadSafe> FPakPlatformFile::GetPakSignatureFile(const TCHAR* InFilename)
-{
-	FScopeLock Lock(&PakSignatureFileCacheLock);
 
+TSharedPtr<const FPakSignatureFile, ESPMode::ThreadSafe> FPakPlatformFile::GetPakSignatureFile(const TCHAR* InFilename)
+{
 	FName FilenameFName(InFilename);
-	if (const TSharedPtr<const struct FPakSignatureFile, ESPMode::ThreadSafe>* SignaturesFile = PakSignatureFileCache.Find(FilenameFName))
 	{
-		return *SignaturesFile;
+		FScopeLock Lock(&PakSignatureFileCacheLock);
+		if (TSharedPtr<const FPakSignatureFile, ESPMode::ThreadSafe>* ExistingSignatureFile = PakSignatureFileCache.Find(FilenameFName))
+		{
+			return *ExistingSignatureFile;
+		}
 	}
 
-	static FRSAKeyHandle PublicKey = InvalidRSAKeyHandle;
-	static bool bInitializedPublicKey = false;
-	if (!bInitializedPublicKey)
+	static FRSAKeyHandle PublicKey = []() -> FRSAKeyHandle
 	{
 		TDelegate<void(TArray<uint8>&, TArray<uint8>&)>& Delegate = FCoreDelegates::GetPakSigningKeysDelegate();
 		if (Delegate.IsBound())
@@ -373,31 +373,37 @@ TSharedPtr<const struct FPakSignatureFile, ESPMode::ThreadSafe> FPakPlatformFile
 			TArray<uint8> Exponent;
 			TArray<uint8> Modulus;
 			Delegate.Execute(Exponent, Modulus);
-			PublicKey = FRSA::CreateKey(Exponent, TArray<uint8>(), Modulus);
+			return FRSA::CreateKey(Exponent, TArray<uint8>(), Modulus);
 		}
-		bInitializedPublicKey = true;
-	}
+		return InvalidRSAKeyHandle;
+	}();
 
-	TSharedPtr<FPakSignatureFile, ESPMode::ThreadSafe> SignaturesFile;
+	TSharedPtr<FPakSignatureFile, ESPMode::ThreadSafe> NewSignatureFile;
 
 	if (PublicKey != InvalidRSAKeyHandle)
 	{
 		FString SignaturesFilename = FPaths::ChangeExtension(InFilename, TEXT("sig"));
-		FArchive* Reader = IFileManager::Get().CreateFileReader(*SignaturesFilename);
+		TUniquePtr<FArchive> Reader(IFileManager::Get().CreateFileReader(*SignaturesFilename));
 		if (Reader != nullptr)
 		{
-			SignaturesFile = MakeShared<FPakSignatureFile, ESPMode::ThreadSafe>();
-			SignaturesFile->Serialize(*Reader);
-			delete Reader;
+			NewSignatureFile = MakeShared<FPakSignatureFile, ESPMode::ThreadSafe>();
+			NewSignatureFile->Serialize(*Reader);
 
-			if (!SignaturesFile->DecryptSignatureAndValidate(PublicKey, InFilename))
+			if (!NewSignatureFile->DecryptSignatureAndValidate(PublicKey, InFilename))
 			{
 				// We don't need to act on this failure as the decrypt function will already have dumped out log messages
 				// and fired the signature check fail handler
-				SignaturesFile.Reset();
+				NewSignatureFile.Reset();
 			}
 
-			PakSignatureFileCache.Add(FilenameFName, SignaturesFile);
+			{
+				FScopeLock Lock(&PakSignatureFileCacheLock);
+				if (TSharedPtr<const FPakSignatureFile, ESPMode::ThreadSafe>* ExistingSignatureFile = PakSignatureFileCache.Find(FilenameFName))
+				{
+					return *ExistingSignatureFile;
+				}
+				PakSignatureFileCache.Add(FilenameFName, NewSignatureFile);
+			}
 		}
 		else
 		{
@@ -406,17 +412,14 @@ TSharedPtr<const struct FPakSignatureFile, ESPMode::ThreadSafe> FPakPlatformFile
 		}
 	}
 
-	return SignaturesFile;
+	return NewSignatureFile;
 }
 
 void FPakPlatformFile::RemoveCachedPakSignaturesFile(const TCHAR* InFilename)
 {
-	FScopeLock Lock(&PakSignatureFileCacheLock);
 	FName FilenameFName(InFilename);
-	if (PakSignatureFileCache.Contains(FilenameFName))
-	{
-		PakSignatureFileCache.Remove(FilenameFName);
-	}
+	FScopeLock Lock(&PakSignatureFileCacheLock);
+	PakSignatureFileCache.Remove(FilenameFName);
 }
 
 DECLARE_DWORD_ACCUMULATOR_STAT(TEXT("PakCache Sync Decrypts (Uncompressed Path)"), STAT_PakCache_SyncDecrypts, STATGROUP_PakFile);
@@ -7856,43 +7859,41 @@ bool FPakPlatformFile::Unmount(const TCHAR* InPakFilename)
 	bool bRemovedContainerFile = false;
 	{
 		FScopeLock ScopedLock(&PakListCritical);
-
-		bool bRemovedPakFile = false;
 		for (int32 PakIndex = 0; PakIndex < PakFiles.Num(); PakIndex++)
 		{
+			FPakListEntry& PakListEntry = PakFiles[PakIndex];
 			if (PakFiles[PakIndex].PakFile->GetFilename() == InPakFilename)
 			{
-				FPakListEntry& PakListEntry = PakFiles[PakIndex];
-				RemoveCachedPakSignaturesFile(*PakListEntry.PakFile->GetFilename());
 				UnmountedPak = MoveTemp(PakListEntry.PakFile);
 				PakFiles.RemoveAt(PakIndex);
-				bRemovedPakFile = true;
 				break;
 			}
 		}
+	}
 
-		if (IoDispatcherFileBackend.IsValid())
-		{
-			if (UnmountedPak)
-			{
-				PackageStoreBackend->Unmount(UnmountedPak->IoContainerHeader.Get());
-			}
-			FString ContainerPath = FPaths::ChangeExtension(InPakFilename, FString());
-			bRemovedContainerFile = IoDispatcherFileBackend->Unmount(*ContainerPath);
-#if WITH_EDITOR
-			if (UnmountedPak && UnmountedPak->OptionalSegmentIoContainerHeader.IsValid())
-			{
-				PackageStoreBackend->Unmount(UnmountedPak->OptionalSegmentIoContainerHeader.Get());
-				FString OptionalSegmentContainerPath = ContainerPath + FPackagePath::GetOptionalSegmentExtensionModifier();
-				IoDispatcherFileBackend->Unmount(*OptionalSegmentContainerPath);
-			}
-#endif
-		}
+	RemoveCachedPakSignaturesFile(*UnmountedPak->GetFilename());
 
+	if (IoDispatcherFileBackend.IsValid())
+	{
 		if (UnmountedPak)
 		{
-			UnmountedPak->Readers.Empty();
+			PackageStoreBackend->Unmount(UnmountedPak->IoContainerHeader.Get());
+		}
+		FString ContainerPath = FPaths::ChangeExtension(InPakFilename, FString());
+		bRemovedContainerFile = IoDispatcherFileBackend->Unmount(*ContainerPath);
+#if WITH_EDITOR
+		if (UnmountedPak && UnmountedPak->OptionalSegmentIoContainerHeader.IsValid())
+		{
+			PackageStoreBackend->Unmount(UnmountedPak->OptionalSegmentIoContainerHeader.Get());
+			FString OptionalSegmentContainerPath = ContainerPath + FPackagePath::GetOptionalSegmentExtensionModifier();
+			IoDispatcherFileBackend->Unmount(*OptionalSegmentContainerPath);
 		}
+#endif
+	}
+
+	if (UnmountedPak)
+	{
+		UnmountedPak->Readers.Empty();
 	}
 #if USE_PAK_PRECACHE
 	if (GPakCache_Enable)
