Fixing crash when unmounting a pak that is not yet in the paklist.
#tests Android Dev Client

diff --git a/Engine/Source/Runtime/PakFile/Private/IPlatformFilePak.cpp b/Engine/Source/Runtime/PakFile/Private/IPlatformFilePak.cpp
--- a/Engine/Source/Runtime/PakFile/Private/IPlatformFilePak.cpp
+++ b/Engine/Source/Runtime/PakFile/Private/IPlatformFilePak.cpp
@@ -7871,7 +7871,10 @@ bool FPakPlatformFile::Unmount(const TCHAR* InPakFilename)
 		}
 	}
 
-	RemoveCachedPakSignaturesFile(*UnmountedPak->GetFilename());
+	if (UnmountedPak)
+	{
+		RemoveCachedPakSignaturesFile(*UnmountedPak->GetFilename());
+	}
 
 	if (IoDispatcherFileBackend.IsValid())
 	{
