Truncation warning fixes for multiple core modules:
CookOnTheFly, CookOnTheFlyNetServer, PakFileUtilities, PakFile

diff --git a/Engine/Source/Developer/CookOnTheFlyNetServer/Private/CookOnTheFlyNetServerBase.cpp b/Engine/Source/Developer/CookOnTheFlyNetServer/Private/CookOnTheFlyNetServerBase.cpp
--- a/Engine/Source/Developer/CookOnTheFlyNetServer/Private/CookOnTheFlyNetServerBase.cpp
+++ b/Engine/Source/Developer/CookOnTheFlyNetServer/Private/CookOnTheFlyNetServerBase.cpp
@@ -139,7 +139,7 @@ bool FCookOnTheFlyClientConnectionBase::ProcessPayload(FArchive& Payload)
 bool FCookOnTheFlyClientConnectionBase::SendMessage(const UE::Cook::FCookOnTheFlyMessage& Message)
 {
 	FBufferArchive Payload;
-	Payload.Reserve(Message.TotalSize());
+	Payload.Reserve(IntCastChecked<int32>(Message.TotalSize()));
 	Payload << const_cast<UE::Cook::FCookOnTheFlyMessage&>(Message);
 	return SendPayload(Payload);
 }
diff --git a/Engine/Source/Developer/PakFileUtilities/Private/PakFileUtilities.cpp b/Engine/Source/Developer/PakFileUtilities/Private/PakFileUtilities.cpp
--- a/Engine/Source/Developer/PakFileUtilities/Private/PakFileUtilities.cpp
+++ b/Engine/Source/Developer/PakFileUtilities/Private/PakFileUtilities.cpp
@@ -108,7 +108,7 @@ bool SignPakFile(const FString& InPakFilename, const FRSAKeyHandle InSigningKey)
 	Buffer.SetNum(BlockSize);
 
 	TArray<TPakChunkHash> Hashes;
-	Hashes.Empty((uint64)(PakFile->TotalSize() / BlockSize));
+	Hashes.Empty(IntCastChecked<int32>(PakFile->TotalSize() / BlockSize));
 
 	while (Remaining > 0)
 	{
@@ -641,7 +641,7 @@ struct FCompressedFileBuffer
 		FileCompressionBlockSize = 0;
 		FileCompressionMethod = CompressionMethod;
 		CompressedBlocks.Reset();
-		CompressedBlocks.AddUninitialized((OriginalSize+CompressionBlockSize-1)/CompressionBlockSize);
+		CompressedBlocks.AddUninitialized(IntCastChecked<int32>((OriginalSize+CompressionBlockSize-1)/CompressionBlockSize));
 	}
 
 	void Empty()
@@ -990,7 +990,7 @@ bool FCompressedFileBuffer::EndCompressFileToWorkingBuffer(const FPakInputPair&
 
 			if (InFile.bNeedEncryption)
 			{
-				int32 EncryptionBlockPadding = Align(TotalCompressedSize, FAES::AESBlockSize);
+				int64 EncryptionBlockPadding = Align(TotalCompressedSize, FAES::AESBlockSize);
 				for (int64 FillIndex = TotalCompressedSize; FillIndex < EncryptionBlockPadding; ++FillIndex)
 				{
 					// Fill the trailing buffer with bytes from file. Note that this is now from a fixed location
@@ -1696,7 +1696,7 @@ bool UncompressCopyFile(FArchive& Dest, FArchive& Source, const FPakEntry& Entry
 	int32 MaxCompressionBlockSize = FCompression::CompressMemoryBound(EntryCompressionMethod, Entry.CompressionBlockSize);
 	for (const FPakCompressedBlock& Block : Entry.CompressionBlocks)
 	{
-		MaxCompressionBlockSize = FMath::Max<int32>(MaxCompressionBlockSize, Block.CompressedEnd - Block.CompressedStart);
+		MaxCompressionBlockSize = FMath::Max<int32>(MaxCompressionBlockSize, IntCastChecked<int32>(Block.CompressedEnd - Block.CompressedStart));
 	}
 
 	int64 WorkingSize = Entry.CompressionBlockSize + MaxCompressionBlockSize;
@@ -1710,10 +1710,10 @@ bool UncompressCopyFile(FArchive& Dest, FArchive& Source, const FPakEntry& Entry
 
 	for (uint32 BlockIndex=0, BlockIndexNum=Entry.CompressionBlocks.Num(); BlockIndex < BlockIndexNum; ++BlockIndex)
 	{
-		uint32 CompressedBlockSize = Entry.CompressionBlocks[BlockIndex].CompressedEnd - Entry.CompressionBlocks[BlockIndex].CompressedStart;
-		uint32 UncompressedBlockSize = (uint32)FMath::Min<int64>(Entry.UncompressedSize - Entry.CompressionBlockSize*BlockIndex, Entry.CompressionBlockSize);
+		int64 CompressedBlockSize = Entry.CompressionBlocks[BlockIndex].CompressedEnd - Entry.CompressionBlocks[BlockIndex].CompressedStart;
+		int64 UncompressedBlockSize = FMath::Min<int64>(Entry.UncompressedSize - Entry.CompressionBlockSize*BlockIndex, Entry.CompressionBlockSize);
 		Source.Seek(Entry.CompressionBlocks[BlockIndex].CompressedStart + (PakFile.GetInfo().HasRelativeCompressedChunkOffsets() ? Entry.Offset : 0));
-		uint32 SizeToRead = Entry.IsEncrypted() ? Align(CompressedBlockSize, FAES::AESBlockSize) : CompressedBlockSize;
+		int64 SizeToRead = Entry.IsEncrypted() ? Align(CompressedBlockSize, FAES::AESBlockSize) : CompressedBlockSize;
 		Source.Serialize(PersistentBuffer, SizeToRead);
 
 		if (Entry.IsEncrypted())
@@ -1727,7 +1727,7 @@ bool UncompressCopyFile(FArchive& Dest, FArchive& Source, const FPakEntry& Entry
 			FAES::DecryptData(PersistentBuffer, SizeToRead, Key->Key);
 		}
 
-		if (!FCompression::UncompressMemory(EntryCompressionMethod, UncompressedBuffer, UncompressedBlockSize, PersistentBuffer, CompressedBlockSize))
+		if (!FCompression::UncompressMemory(EntryCompressionMethod, UncompressedBuffer, IntCastChecked<int32>(UncompressedBlockSize), PersistentBuffer, IntCastChecked<int32>(CompressedBlockSize)))
 		{
 			return false;
 		}
@@ -2429,7 +2429,7 @@ void FPakWriterContext::EndCompress(FOutputPakFileEntry* Entry)
 					// Check the compression ratio, if it's too low just store uncompressed. Also take into account read size
 					// if we still save 64KB it's probably worthwhile compressing, as that saves a file read operation in the runtime.
 					// TODO: drive this threshold from the command line
-					float PercentLess = ((float)Entry->GetCompressedBuffer().TotalCompressedSize / (Entry->OriginalFileSize / 100.f));
+					float PercentLess = ((float)Entry->GetCompressedBuffer().TotalCompressedSize / ((float)Entry->OriginalFileSize / 100.f));
 					bNotEnoughCompression = (PercentLess > 90.f) && ((Entry->OriginalFileSize - Entry->GetCompressedBuffer().TotalCompressedSize) < 65536);
 				}
 			}
@@ -2534,7 +2534,7 @@ void FPakWriterContext::CompressionThreadFunc()
 
 void FPakWriterContext::WriterThreadFunc()
 {
-	const uint32 RequiredPatchPadding = CmdLineParameters.PatchFilePadAlign;
+	const int64 RequiredPatchPadding = CmdLineParameters.PatchFilePadAlign;
 
 	for (;;)
 	{
@@ -2783,7 +2783,7 @@ bool FPakWriterContext::Flush()
 	CompressionQueueEntryAddedEvent->Trigger();
 	WriteQueueEntryAddedEvent->Trigger();
 
-	const uint32 RequiredPatchPadding = CmdLineParameters.PatchFilePadAlign;
+	const int64 RequiredPatchPadding = CmdLineParameters.PatchFilePadAlign;
 
 	for (TUniquePtr<FOutputPakFile>& OutputPakFile : OutputPakFiles)
 	{
@@ -2823,7 +2823,7 @@ bool FPakWriterContext::Flush()
 
 				if (OutputEntry.InputPair.bNeedsCompression && OutputEntry.CompressionMethod != NAME_None)
 				{
-					float PercentLess = ((float)OutputEntry.Entry.Info.Size / (OutputEntry.Entry.Info.UncompressedSize / 100.f));
+					float PercentLess = ((float)OutputEntry.Entry.Info.Size / ((float)OutputEntry.Entry.Info.UncompressedSize / 100.f));
 					if (OutputEntry.InputPair.SuggestedOrder < MAX_uint64)
 					{
 						UE_LOG(LogPakFile, Verbose, TEXT("Added compressed %sfile \"%s\", %.2f%% of original size. Compressed with %s, Size %lld bytes, Original Size %lld bytes (order %llu)."), EncryptedString, *OutputEntry.Entry.Filename, PercentLess, *OutputEntry.CompressionMethod.ToString(), OutputEntry.Entry.Info.Size, OutputEntry.Entry.Info.UncompressedSize, OutputEntry.InputPair.SuggestedOrder);
@@ -2920,7 +2920,7 @@ bool FPakWriterContext::Flush()
 		UE_LOG(LogPakFile, Display, TEXT("FullDirectoryIndex size: %d bytes"), Footer.FullDirectoryIndexSize);
 		if (OutputPakFile->TotalUncompressedSize)
 		{
-			float PercentLess = ((float)OutputPakFile->TotalCompressedSize / (OutputPakFile->TotalUncompressedSize / 100.f));
+			float PercentLess = ((float)OutputPakFile->TotalCompressedSize / ((float)OutputPakFile->TotalUncompressedSize / 100.f));
 			UE_LOG(LogPakFile, Display, TEXT("Compression summary: %.2f%% of original size. Compressed Size %lld bytes, Original Size %lld bytes. "), PercentLess, OutputPakFile->TotalCompressedSize, OutputPakFile->TotalUncompressedSize);
 		}
 
@@ -3301,8 +3301,8 @@ bool ListFilesInPak(const TCHAR * InPakFilename, int64 SizeFilter, bool bInclude
 			{
 				if (InspectChunkRanges.Num() > 0)
 				{
-					int32 FirstChunk = Entry.Offset / (64 * 1024);
-					int32 LastChunk = (Entry.Offset + Entry.Size) / (64 * 1024);
+					int32 FirstChunk = IntCastChecked<int32>(Entry.Offset) / (64 * 1024);
+					int32 LastChunk = IntCastChecked<int32>(Entry.Offset + Entry.Size) / (64 * 1024);
 
 					for (int32 Chunk = FirstChunk; Chunk <= LastChunk; ++Chunk)
 					{
@@ -3418,7 +3418,7 @@ bool AuditPakFiles( const FString& InputPath, bool bOnlyDeleted, const FString&
 	{
 		FString PakFilename;
 		int32 PakPriority;
-		int32 Size;
+		int64 Size;
 	};
 	TMap<FString, FFilePakRevision> FileRevisions;
 	TMap<FString, FFilePakRevision> DeletedRevisions;
@@ -3680,7 +3680,7 @@ bool AuditPakFiles( const FString& InputPath, bool bOnlyDeleted, const FString&
 	//write seek summary
 	if( bSortByOrdering && bHasOpenOrder && NumReads > 0 )
 	{
-		UE_LOG( LogPakFile, Display, TEXT("%d guaranteed seeks out of %d files read (%.2f%%) with the given open order"), NumSeeks, NumReads, (NumSeeks*100.0f)/NumReads );
+		UE_LOG( LogPakFile, Display, TEXT("%d guaranteed seeks out of %d files read (%.2f%%) with the given open order"), NumSeeks, NumReads, (float)(NumSeeks*100.0) / (float)NumReads );
 	}
 
 	return true;
@@ -3803,7 +3803,7 @@ bool ShowCompressionBlockCRCs( const TCHAR* InPakFileName, TArray<int64>& InOffs
 		int32 MaxCompressionBlockSize = FCompression::CompressMemoryBound(EntryCompressionMethod, Entry->CompressionBlockSize);
 		for (const FPakCompressedBlock& Block : Entry->CompressionBlocks)
 		{
-			MaxCompressionBlockSize = FMath::Max<int32>(MaxCompressionBlockSize, Block.CompressedEnd - Block.CompressedStart);
+			MaxCompressionBlockSize = FMath::Max<int32>(MaxCompressionBlockSize, IntCastChecked<int32>(Block.CompressedEnd - Block.CompressedStart));
 		}
 
 		int64 WorkingSize = Entry->CompressionBlockSize + MaxCompressionBlockSize;
@@ -3811,10 +3811,9 @@ bool ShowCompressionBlockCRCs( const TCHAR* InPakFileName, TArray<int64>& InOffs
 
 		for (uint32 BlockIndex=0, BlockIndexNum=Entry->CompressionBlocks.Num(); BlockIndex < BlockIndexNum; ++BlockIndex)
 		{
-			uint32 CompressedBlockSize = Entry->CompressionBlocks[BlockIndex].CompressedEnd - Entry->CompressionBlocks[BlockIndex].CompressedStart;
-			uint32 UncompressedBlockSize = (uint32)FMath::Min<int64>(Entry->UncompressedSize - Entry->CompressionBlockSize*BlockIndex, Entry->CompressionBlockSize);
+			int64 CompressedBlockSize = Entry->CompressionBlocks[BlockIndex].CompressedEnd - Entry->CompressionBlocks[BlockIndex].CompressedStart;
 			PakReader->Seek(Entry->CompressionBlocks[BlockIndex].CompressedStart + (PakFile.GetInfo().HasRelativeCompressedChunkOffsets() ? Entry->Offset : 0));
-			uint32 SizeToRead = Entry->IsEncrypted() ? Align(CompressedBlockSize, FAES::AESBlockSize) : CompressedBlockSize;
+			int64 SizeToRead = Entry->IsEncrypted() ? Align(CompressedBlockSize, FAES::AESBlockSize) : CompressedBlockSize;
 			PakReader->Serialize(PersistentBuffer, SizeToRead);
 
 			if (Entry->IsEncrypted())
@@ -3836,8 +3835,9 @@ bool ShowCompressionBlockCRCs( const TCHAR* InPakFileName, TArray<int64>& InOffs
 			}
 
 			// compute checksum and log out the block information
-			const uint32 BlockCrc32 = FCrc::MemCrc32( PersistentBuffer, CompressedBlockSize );
-			const FString HexBytes = BytesToHex(PersistentBuffer,FMath::Min(CompressedBlockSize,32U));
+			uint32 TruncatedCompressedBlockSize = IntCastChecked<uint32>(CompressedBlockSize);
+			const uint32 BlockCrc32 = FCrc::MemCrc32(PersistentBuffer, TruncatedCompressedBlockSize);
+			const FString HexBytes = BytesToHex(PersistentBuffer,FMath::Min(TruncatedCompressedBlockSize,32U));
 			UE_LOG(LogPakFile, Display, TEXT("    Block:%-6d  ProcessedSize: %-6d  DecompressionRawSize: %-6d  Crc32: %-12u [%s...]"), BlockIndex, ProcessedSize, CompressedBlockSize, BlockCrc32, *HexBytes );
 		}
 
@@ -4549,7 +4549,7 @@ bool FileIsIdentical(FString SourceFile, FString DestFilename, const FFileInfo*
 	return true;
 }
 
-float GetFragmentationPercentage(const TArray<FPakInputPair>& FilesToPak, const TBitArray<>& IncludeBitMask, int32 MaxAdjacentOrderDiff, bool bConsiderSecondaryFiles)
+float GetFragmentationPercentage(const TArray<FPakInputPair>& FilesToPak, const TBitArray<>& IncludeBitMask, int64 MaxAdjacentOrderDiff, bool bConsiderSecondaryFiles)
 {
 	uint64 PrevOrder = MAX_uint64;
 	bool bPrevBit = false;
@@ -4565,7 +4565,7 @@ float GetFragmentationPercentage(const TArray<FPakInputPair>& FilesToPak, const
 		uint64 CurrentOrder = FilesToPak[i].SuggestedOrder;
 		bool bCurrentBit = IncludeBitMask[i];
 		uint64 OrderDiff = CurrentOrder - PrevOrder;
-		if (OrderDiff > MaxAdjacentOrderDiff || bCurrentBit != bPrevBit)
+		if (OrderDiff > (uint64)(MaxAdjacentOrderDiff) || bCurrentBit != bPrevBit)
 		{
 			DiffCount++;
 		}
@@ -4629,7 +4629,7 @@ int32 AddOrphanedFiles(const TArray<FPakInputPair>& FilesToPak, const TArray<int
 }
 
 
-bool DoGapFillingIteration(const TArray<FPakInputPair>& FilesToPak, const TArray<int64>& FileSizes, const TArray<int32>& UAssetToUexpMapping, TBitArray<>& InOutIncludeFilesMask, int64 MaxGapSizeBytes, int32 MaxAdjacentOrderDiff, bool bForce, int64 MaxPatchSize = 0, bool bFillPrimaryOrderFiles = true, bool bFillSecondaryOrderFiles = true)
+bool DoGapFillingIteration(const TArray<FPakInputPair>& FilesToPak, const TArray<int64>& FileSizes, const TArray<int32>& UAssetToUexpMapping, TBitArray<>& InOutIncludeFilesMask, int64 MaxGapSizeBytes, int64 MaxAdjacentOrderDiff, bool bForce, int64 MaxPatchSize = 0, bool bFillPrimaryOrderFiles = true, bool bFillSecondaryOrderFiles = true)
 {
 	TBitArray<> IncludeFilesMask = InOutIncludeFilesMask;
 
@@ -4751,7 +4751,7 @@ bool DoIncrementalGapFilling(const TArray<FPakInputPair>& FilesToPak, const TArr
 		else
 		{
 			// Try with 75% of the max gap size
-			bSuccess |= DoGapFillingIteration(FilesToPak, FileSizes, UAssetToUexpMapping, IncludeFilesMask, GapSize*0.75, MaxAdjacentOrderDiff, false, MaxPatchSize, bFillPrimaryOrderFiles, bFillSecondaryOrderFiles);
+			bSuccess |= DoGapFillingIteration(FilesToPak, FileSizes, UAssetToUexpMapping, IncludeFilesMask, (int64)((double)GapSize*0.75), MaxAdjacentOrderDiff, false, MaxPatchSize, bFillPrimaryOrderFiles, bFillSecondaryOrderFiles);
 			break;
 		}
 		GapSize *= 2;
@@ -4790,7 +4790,7 @@ void ApplyGapFilling(const TArray<FPakInputPair>& FilesToPak, const TArray<int64
 	{
 		UE_CLOG(SeekOptParams.MaxInflationPercent == 0.0f, LogPakFile, Fatal, TEXT("ESeekOptMode::Incremental* requires MaxInflationPercent > 0.0"));
 		int64 PassMaxPatchSize[3];
-		PassMaxPatchSize[0] = OriginalPatchSize + (IncrementalMaxPatchSize - OriginalPatchSize) * 0.9f;
+		PassMaxPatchSize[0] = OriginalPatchSize + (int64)((double)(IncrementalMaxPatchSize - OriginalPatchSize) * 0.9);
 		PassMaxPatchSize[1] = IncrementalMaxPatchSize;
 		PassMaxPatchSize[2] = IncrementalMaxPatchSize;
 		for (int32 i = 0; i < 3; i++)
diff --git a/Engine/Source/Runtime/CookOnTheFly/Private/CookOnTheFlyServerConnection.cpp b/Engine/Source/Runtime/CookOnTheFly/Private/CookOnTheFlyServerConnection.cpp
--- a/Engine/Source/Runtime/CookOnTheFly/Private/CookOnTheFlyServerConnection.cpp
+++ b/Engine/Source/Runtime/CookOnTheFly/Private/CookOnTheFlyServerConnection.cpp
@@ -120,7 +120,7 @@ public:
 		RequestHeader.Timestamp		= FDateTime::UtcNow().GetTicks();
 
 		FBufferArchive RequestPayload;
-		RequestPayload.Reserve(Request.TotalSize());
+		RequestPayload.Reserve(IntCastChecked<int32>(Request.TotalSize()));
 
 		RequestPayload << Request;
 
diff --git a/Engine/Source/Runtime/PakFile/Private/IPlatformFilePak.cpp b/Engine/Source/Runtime/PakFile/Private/IPlatformFilePak.cpp
--- a/Engine/Source/Runtime/PakFile/Private/IPlatformFilePak.cpp
+++ b/Engine/Source/Runtime/PakFile/Private/IPlatformFilePak.cpp
@@ -189,7 +189,7 @@ static FAutoConsoleCommand CVar_TestRegisterEncryptionKey(
 TPakChunkHash ComputePakChunkHash(const void* InData, int64 InDataSizeInBytes)
 {
 #if PAKHASH_USE_CRC
-	return FCrc::MemCrc32(InData, InDataSizeInBytes);
+	return FCrc::MemCrc32(InData, IntCastChecked<int32>(InDataSizeInBytes));
 #else
 	FSHAHash Hash;
 	FSHA1::HashBuffer(InData, InDataSizeInBytes, Hash.Hash);
@@ -427,11 +427,12 @@ DECLARE_FLOAT_ACCUMULATOR_STAT(TEXT("PakCache Decrypt Time"), STAT_PakCache_Decr
 DECLARE_DWORD_ACCUMULATOR_STAT(TEXT("PakCache Async Decrypts (Compressed Path)"), STAT_PakCache_CompressedDecrypts, STATGROUP_PakFile);
 DECLARE_DWORD_ACCUMULATOR_STAT(TEXT("PakCache Async Decrypts (Uncompressed Path)"), STAT_PakCache_UncompressedDecrypts, STATGROUP_PakFile);
 
-void DecryptData(uint8* InData, uint32 InDataSize, FGuid InEncryptionKeyGuid)
+void DecryptData(uint8* InData, uint64 InDataSize, FGuid InEncryptionKeyGuid)
 {
+	uint32 DataSize = IntCastChecked<uint32>(InDataSize);
 	if (FPakPlatformFile::GetPakCustomEncryptionDelegate().IsBound())
 	{
-		FPakPlatformFile::GetPakCustomEncryptionDelegate().Execute(InData, InDataSize, InEncryptionKeyGuid);
+		FPakPlatformFile::GetPakCustomEncryptionDelegate().Execute(InData, DataSize, InEncryptionKeyGuid);
 	}
 	else
 	{
@@ -439,7 +440,7 @@ void DecryptData(uint8* InData, uint32 InDataSize, FGuid InEncryptionKeyGuid)
 		FAES::FAESKey Key;
 		FPakPlatformFile::GetPakEncryptionKey(Key, InEncryptionKeyGuid);
 		check(Key.IsValid());
-		FAES::DecryptData(InData, InDataSize, Key);
+		FAES::DecryptData(InData, DataSize, Key);
 	}
 }
 
@@ -1527,7 +1528,7 @@ public:
 				return nullptr;
 			}
 			CachedPakData.Add(FPakData(InActualPakFile, Handle, File, PakFileSize));
-			PakIndexPtr = &CachedPaks.Add(InActualPakFile, CachedPakData.Num() - 1);
+			PakIndexPtr = &CachedPaks.Add(InActualPakFile, static_cast<uint16>(CachedPakData.Num() - 1));
 			FPakData& Pak = CachedPakData[*PakIndexPtr];
 
 
@@ -1633,7 +1634,7 @@ private: // below here we assume CachedFilesScopeLock until we get to the next s
 
 		int64 FirstByte = AlignDown(Offset, PAK_CACHE_GRANULARITY);
 		int64 LastByte = Align(Offset + Size, PAK_CACHE_GRANULARITY) - 1;
-		uint32 NumBits = (PAK_CACHE_GRANULARITY + LastByte - FirstByte) / PAK_CACHE_GRANULARITY;
+		uint32 NumBits = IntCastChecked<uint32>((PAK_CACHE_GRANULARITY + LastByte - FirstByte) / PAK_CACHE_GRANULARITY);
 		uint32 NumQWords = (NumBits + 63) >> 6;
 		InFlightOrDone.Reset();
 		InFlightOrDone.AddZeroed(NumQWords);
@@ -1702,7 +1703,7 @@ private: // below here we assume CachedFilesScopeLock until we get to the next s
 
 		int64 FirstByte = AlignDown(Offset, PAK_CACHE_GRANULARITY);
 		int64 LastByte = Align(Offset + Request.Size, PAK_CACHE_GRANULARITY) - 1;
-		uint32 NumBits = (PAK_CACHE_GRANULARITY + LastByte - FirstByte) / PAK_CACHE_GRANULARITY;
+		uint32 NumBits = IntCastChecked<uint32>((PAK_CACHE_GRANULARITY + LastByte - FirstByte) / PAK_CACHE_GRANULARITY);
 		uint32 NumQWords = (NumBits + 63) >> 6;
 		InFlightOrDone.Reset();
 		InFlightOrDone.AddZeroed(NumQWords);
@@ -1858,7 +1859,7 @@ private: // below here we assume CachedFilesScopeLock until we get to the next s
 		if (GPakCache_UseNewTrim && !bDiscardAll)
 		{
 			StartPakIndex = 0;
-			uint16 EndPakIndex = CachedPakData.Num();
+			uint16 EndPakIndex = IntCastChecked<uint16>(CachedPakData.Num());
 
 			// TODO: remove this array, add a bool to the cache object
 			TArray<bool> CacheVisitedAlready;
@@ -2018,7 +2019,7 @@ private: // below here we assume CachedFilesScopeLock until we get to the next s
 			else
 			{
 				StartPakIndex = 0;
-				EndPakIndex = CachedPakData.Num();
+				EndPakIndex = IntCastChecked<uint16>(CachedPakData.Num());
 			}
 
 			// CachedFilesScopeLock is locked
@@ -2285,7 +2286,7 @@ private: // below here we assume CachedFilesScopeLock until we get to the next s
 		int64 LastByte = FMath::Min(Align(FirstByte + (GPakCache_MaxRequestSizeToLowerLevelKB * 1024), PAK_CACHE_GRANULARITY) - 1, Pak.TotalSize - 1);
 		check(FirstByte >= 0 && LastByte < Pak.TotalSize && LastByte >= 0 && LastByte >= FirstByte);
 
-		uint32 NumBits = (PAK_CACHE_GRANULARITY + LastByte - FirstByte) / PAK_CACHE_GRANULARITY;
+		uint32 NumBits = IntCastChecked<uint32>((PAK_CACHE_GRANULARITY + LastByte - FirstByte) / PAK_CACHE_GRANULARITY);
 		uint32 NumQWords = (NumBits + 63) >> 6;
 
 		static TArray<uint64> InFlightOrDone;
@@ -3182,7 +3183,7 @@ static void WaitPrecache(const TArray<FString>& Args)
 	}
 	Loads = FPakPrecacher::Get().GetLoads() - Loads;
 	LoadSize = FPakPrecacher::Get().GetLoadSize() - LoadSize;
-	float TimeSpent = FPlatformTime::Seconds() - StartTime;
+	float TimeSpent = FloatCastChecked<float>(FPlatformTime::Seconds() - StartTime, UE::LWC::DefaultFloatPrecision);
 	float LoadSizeMB = float(LoadSize) / (1024.0f * 1024.0f);
 	float MBs = LoadSizeMB / TimeSpent;
 	UE_LOG(LogPakFile, Log, TEXT("Loaded %4d blocks (align %4dKB) totalling %7.2fMB in %4.2fs   = %6.2fMB/s"), Loads, PAK_CACHE_GRANULARITY / 1024, LoadSizeMB, TimeSpent, MBs);
@@ -3338,7 +3339,7 @@ public:
 			}
 			else
 			{
-				WaitEvent->Wait(TimeLimitSeconds * 1000.0f);
+				WaitEvent->Wait(static_cast<uint32>(TimeLimitSeconds * 1000.0f));
 			}
 			FScopeLock Lock(&FPakReadRequestEvent);
 			FPlatformProcess::ReturnSynchEventToPool(WaitEvent);
@@ -3593,7 +3594,7 @@ public:
 			}
 			else
 			{
-				WaitEvent->Wait(TimeLimitSeconds * 1000.0f);
+				WaitEvent->Wait(static_cast<uint32>(TimeLimitSeconds * 1000.0f));
 			}
 			FScopeLock Lock(&FPakReadRequestEvent);
 			FPlatformProcess::ReturnSynchEventToPool(WaitEvent);
@@ -3737,7 +3738,7 @@ void FPakPrecacher::StartSignatureCheck(bool bWasCanceled, IAsyncReadRequest* Re
 
 void FPakPrecacher::DoSignatureCheck(bool bWasCanceled, IAsyncReadRequest* Request, int32 Index)
 {
-	int64 SignatureIndex = -1;
+	int32 SignatureIndex = -1;
 	int64 NumSignaturesToCheck = 0;
 	const uint8* Data = nullptr;
 	int64 RequestSize = 0;
@@ -3768,7 +3769,7 @@ void FPakPrecacher::DoSignatureCheck(bool bWasCanceled, IAsyncReadRequest* Reque
 		RequestSize = RequestToLower.RequestSize;
 		PakIndex = GetRequestPakIndex(Block.OffsetAndPakIndex);
 		Data = RequestToLower.Memory;
-		SignatureIndex = RequestOffset / FPakInfo::MaxChunkDataSize;
+		SignatureIndex = IntCastChecked<int32>(RequestOffset / FPakInfo::MaxChunkDataSize);
 
 		FPakData& PakData = CachedPakData[PakIndex];
 		PrincipalSignatureHash = PakData.Signatures->DecryptedHash;
@@ -3999,8 +4000,8 @@ public:
 		{
 			FScopeLock ScopedLock(&CriticalSection);
 			check(Blocks.Num());
-			int32 FirstBlock = Offset / FileEntry.CompressionBlockSize;
-			int32 LastBlock = (Offset + BytesToRead - 1) / FileEntry.CompressionBlockSize;
+			int32 FirstBlock = IntCastChecked<int32>(Offset / FileEntry.CompressionBlockSize);
+			int32 LastBlock = IntCastChecked<int32>((Offset + BytesToRead - 1) / FileEntry.CompressionBlockSize);
 
 			check(FirstBlock >= 0 && FirstBlock < Blocks.Num() && LastBlock >= 0 && LastBlock < Blocks.Num() && FirstBlock <= LastBlock);
 
@@ -4037,7 +4038,7 @@ public:
 		FCachedAsyncBlock& Block = GetBlock(BlockIndex);
 		Block.bInFlight = true;
 		check(!Block.RawRequest && !Block.Processed && !Block.Raw && !Block.CPUWorkGraphEvent.GetReference() && !Block.ProcessedSize && !Block.RawSize && !Block.bCPUWorkIsComplete);
-		Block.RawSize = FileEntry.CompressionBlocks[BlockIndex].CompressedEnd - FileEntry.CompressionBlocks[BlockIndex].CompressedStart;
+		Block.RawSize = IntCastChecked<int32>(FileEntry.CompressionBlocks[BlockIndex].CompressedEnd - FileEntry.CompressionBlocks[BlockIndex].CompressedStart);
 		Block.DecompressionRawSize = Block.RawSize;
 		if (FileEntry.IsEncrypted())
 		{
@@ -4099,8 +4100,8 @@ public:
 			bool bCorrupted = false;
 			if (GPakCache_ForceDecompressionFails && FMath::FRand() < 0.001f)
 			{
-				int32 CorruptOffset = FMath::Clamp(int32(FMath::FRandRange(0.f, Block.RawSize - 1.f)), 0, Block.RawSize - 1);
-				uint8 CorruptValue = uint8(FMath::Clamp(int32(FMath::FRandRange(0.f, 255.f)), 0, 255));
+				int32 CorruptOffset = FMath::Clamp(FMath::RandRange(0, Block.RawSize - 1), 0, Block.RawSize - 1);
+				uint8 CorruptValue = uint8(FMath::Clamp(FMath::RandRange(0, 255), 0, 255));
 				if (Block.Raw[CorruptOffset] != CorruptValue)
 				{
 					UE_LOG(LogPakFile, Error, TEXT("Forcing corruption of decompression source data (predecryption) to verify panic read recovery.  Offset = %d, Value = 0x%x"), CorruptOffset, int32(CorruptValue));
@@ -4311,8 +4312,8 @@ public:
 
 		check(LiveRequests.Contains(Req));
 		LiveRequests.Remove(Req);
-		int32 FirstBlock = Offset / FileEntry.CompressionBlockSize;
-		int32 LastBlock = (Offset + BytesToRead - 1) / FileEntry.CompressionBlockSize;
+		int32 FirstBlock = IntCastChecked<int32>(Offset / FileEntry.CompressionBlockSize);
+		int32 LastBlock = IntCastChecked<int32>((Offset + BytesToRead - 1) / FileEntry.CompressionBlockSize);
 		check(FirstBlock >= 0 && FirstBlock < Blocks.Num() && LastBlock >= 0 && LastBlock < Blocks.Num() && FirstBlock <= LastBlock);
 
 		for (int32 BlockIndex = FirstBlock; BlockIndex <= LastBlock; BlockIndex++)
@@ -4345,8 +4346,8 @@ public:
 		check(!bHasCancelledRef);
 		bHasCancelledRef = true;
 		check(LiveRequests.Contains(Req));
-		int32 FirstBlock = Offset / FileEntry.CompressionBlockSize;
-		int32 LastBlock = (Offset + BytesToRead - 1) / FileEntry.CompressionBlockSize;
+		int32 FirstBlock = IntCastChecked<int32>(Offset / FileEntry.CompressionBlockSize);
+		int32 LastBlock = IntCastChecked<int32>((Offset + BytesToRead - 1) / FileEntry.CompressionBlockSize);
 		check(FirstBlock >= 0 && FirstBlock < Blocks.Num() && LastBlock >= 0 && LastBlock < Blocks.Num() && FirstBlock <= LastBlock);
 
 		for (int32 BlockIndex = FirstBlock; BlockIndex <= LastBlock; BlockIndex++)
@@ -4382,8 +4383,8 @@ public:
 	void GatherResults(uint8* Memory, int64 Offset, int64 BytesToRead)
 	{
 		// CriticalSection is locked
-		int32 FirstBlock = Offset / FileEntry.CompressionBlockSize;
-		int32 LastBlock = (Offset + BytesToRead - 1) / FileEntry.CompressionBlockSize;
+		int32 FirstBlock = IntCastChecked<int32>(Offset / FileEntry.CompressionBlockSize);
+		int32 LastBlock = IntCastChecked<int32>((Offset + BytesToRead - 1) / FileEntry.CompressionBlockSize);
 		check(FirstBlock >= 0 && FirstBlock < Blocks.Num() && LastBlock >= 0 && LastBlock < Blocks.Num() && FirstBlock <= LastBlock);
 
 		for (int32 BlockIndex = FirstBlock; BlockIndex <= LastBlock; BlockIndex++)
@@ -4456,8 +4457,8 @@ bool FPakProcessedReadRequest::CheckCompletion(const FPakEntry& FileEntry, int32
 			return false;
 		}
 	}
-	int32 FirstBlock = Offset / FileEntry.CompressionBlockSize;
-	int32 LastBlock = (Offset + BytesToRead - 1) / FileEntry.CompressionBlockSize;
+	int32 FirstBlock = IntCastChecked<int32>(Offset / FileEntry.CompressionBlockSize);
+	int32 LastBlock = IntCastChecked<int32>((Offset + BytesToRead - 1) / FileEntry.CompressionBlockSize);
 	check(FirstBlock >= 0 && FirstBlock < Blocks.Num() && LastBlock >= 0 && LastBlock < Blocks.Num() && FirstBlock <= LastBlock);
 
 	for (int32 MyBlockIndex = FirstBlock; MyBlockIndex <= LastBlock; MyBlockIndex++)
@@ -4975,7 +4976,7 @@ public:
 	void Serialize(int64 DesiredPosition, void* V, int64 Length)
 	{
 		const int32 CompressionBlockSize = PakEntry.CompressionBlockSize;
-		uint32 CompressionBlockIndex = DesiredPosition / CompressionBlockSize;
+		uint32 CompressionBlockIndex = (uint32)(DesiredPosition / CompressionBlockSize);
 		uint8* WorkingBuffers[2];
 		int64 DirectCopyStart = DesiredPosition % PakEntry.CompressionBlockSize;
 		FAsyncTask<FPakUncompressTask> UncompressTask;
@@ -4996,8 +4997,8 @@ public:
 		{
 			// an amount to extra allocate, in case one block's compressed size is bigger than GetMaximumCompressedSize
 			// @todo this should not be needed, can it be removed?
-			float SlopMultiplier = 1.1f;
-			WorkingBufferRequiredSize = (int64)( WorkingBufferRequiredSize * SlopMultiplier );
+			double SlopMultiplier = 1.1f;
+			WorkingBufferRequiredSize = (int64)((double)WorkingBufferRequiredSize * SlopMultiplier );
 		}
 
 		WorkingBufferRequiredSize = EncryptionPolicy::AlignReadRequest(WorkingBufferRequiredSize);
@@ -5056,9 +5057,9 @@ public:
 					// Block can be decompressed directly into output buffer
 					TaskDetails.CompressionFormat = CompressionMethod;
 					TaskDetails.UncompressedBuffer = (uint8*)V;
-					TaskDetails.UncompressedSize = UncompressedBlockSize;
+					TaskDetails.UncompressedSize = IntCastChecked<int32>(UncompressedBlockSize);
 					TaskDetails.CompressedBuffer = WorkingBuffers[CompressionBlockIndex & 1];
-					TaskDetails.CompressedSize = CompressedBlockSize;
+					TaskDetails.CompressedSize = IntCastChecked<int32>(CompressedBlockSize);
 					TaskDetails.CopyOut = nullptr;
 					ScratchSpace->LastDecompressedBlock = 0xFFFFFFFF;
 					ScratchSpace->LastPakIndexHash = FSHAHash();
@@ -5069,9 +5070,9 @@ public:
 					// Block needs to be copied from a working buffer
 					TaskDetails.CompressionFormat = CompressionMethod;
 					TaskDetails.UncompressedBuffer = ScratchSpace->TempBuffer.Get();
-					TaskDetails.UncompressedSize = UncompressedBlockSize;
+					TaskDetails.UncompressedSize = IntCastChecked<int32>(UncompressedBlockSize);
 					TaskDetails.CompressedBuffer = WorkingBuffers[CompressionBlockIndex & 1];
-					TaskDetails.CompressedSize = CompressedBlockSize;
+					TaskDetails.CompressedSize = IntCastChecked<int32>(CompressedBlockSize);
 					TaskDetails.CopyOut = V;
 					TaskDetails.CopyOffset = DirectCopyStart;
 					TaskDetails.CopyLength = WriteSize;
@@ -5542,7 +5543,7 @@ bool FPakFile::LoadIndexInternal(FArchive& Reader)
 	}
 	TArray<uint8> PrimaryIndexData;
 	Reader.Seek(Info.IndexOffset);
-	PrimaryIndexData.SetNum(Info.IndexSize);
+	PrimaryIndexData.SetNum(IntCastChecked<int32>(Info.IndexSize));
 	{
 		SCOPED_BOOT_TIMING("PakFile_LoadPrimaryIndex");
 		Reader.Serialize(PrimaryIndexData.GetData(), Info.IndexSize);
@@ -5677,7 +5678,7 @@ bool FPakFile::LoadIndexInternal(FArchive& Reader)
 			return false;
 		}
 		Reader.Seek(PathHashIndexOffset);
-		PathHashIndexData.SetNum(PathHashIndexSize);
+		PathHashIndexData.SetNum(IntCastChecked<int32>(PathHashIndexSize));
 		{
 			SCOPED_BOOT_TIMING("PakFile_LoadPathHashIndex");
 			Reader.Serialize(PathHashIndexData.GetData(), PathHashIndexSize);
@@ -5734,7 +5735,7 @@ bool FPakFile::LoadIndexInternal(FArchive& Reader)
 		}
 		TArray<uint8> FullDirectoryIndexData;
 		Reader.Seek(FullDirectoryIndexOffset );
-		FullDirectoryIndexData.SetNum(FullDirectoryIndexSize);
+		FullDirectoryIndexData.SetNum(IntCastChecked<int32>(FullDirectoryIndexSize));
 		{
 			SCOPED_BOOT_TIMING("PakFile_LoadDirectoryIndex");
 			Reader.Serialize(FullDirectoryIndexData.GetData(), FullDirectoryIndexSize);
@@ -5809,7 +5810,7 @@ bool FPakFile::LoadLegacyIndex(FArchive& Reader)
 	}
 
 	Reader.Seek(Info.IndexOffset);
-	IndexData.SetNum(Info.IndexSize);
+	IndexData.SetNum(IntCastChecked<int32>(Info.IndexSize));
 	Reader.Serialize(IndexData.GetData(), Info.IndexSize);
 
 	FSHAHash ComputedHash;
@@ -6477,7 +6478,7 @@ bool FPakFile::EncodePakEntry(FArchive& Ar, const FPakEntry& InPakEntry, const F
 		{
 			for (int CompressionBlockIndex = 0; CompressionBlockIndex < PakEntry.CompressionBlocks.Num(); ++CompressionBlockIndex)
 			{
-				uint32 Value = PakEntry.CompressionBlocks[CompressionBlockIndex].CompressedEnd - PakEntry.CompressionBlocks[CompressionBlockIndex].CompressedStart;
+				uint32 Value = IntCastChecked<uint32>(PakEntry.CompressionBlocks[CompressionBlockIndex].CompressedEnd - PakEntry.CompressionBlocks[CompressionBlockIndex].CompressedStart);
 				Ar << Value;
 			}
 		}
@@ -6583,7 +6584,7 @@ void FPakFile::DecodePakEntry(const uint8* SourcePtr, FPakEntry& OutEntry, const
 		// Per the comment in Encode, if CompressionBlocksCount == 1, we use UncompressedSize for CompressionBlockSize
 		if (CompressionBlocksCount == 1)
 		{
-			OutEntry.CompressionBlockSize = OutEntry.UncompressedSize;
+			OutEntry.CompressionBlockSize = IntCastChecked<uint32>(OutEntry.UncompressedSize);
 		}
 		ensure(OutEntry.CompressionBlockSize != 0);
 	}
@@ -7556,18 +7557,18 @@ bool FPakFile::bSomePakNeedsPruning = false;
 #endif
 
 #if !UE_BUILD_SHIPPING
-int32 GetRecursiveAllocatedSize(const FPakDirectory& Index)
+uint64 GetRecursiveAllocatedSize(const FPakDirectory& Index)
 {
-	int32 Size = Index.GetAllocatedSize();
+	uint64 Size = Index.GetAllocatedSize();
 	for (const TPair<FString, FPakEntryLocation>& kvpair : Index)
 	{
 		Size += kvpair.Key.GetAllocatedSize();
 	}
 	return Size;
 }
-int32 GetRecursiveAllocatedSize(const FPakFile::FDirectoryIndex& Index)
+uint64 GetRecursiveAllocatedSize(const FPakFile::FDirectoryIndex& Index)
 {
-	int32 Size = Index.GetAllocatedSize();
+	uint64 Size = Index.GetAllocatedSize();
 	for (const TPair<FString, FPakDirectory>& kvpair : Index)
 	{
 		Size += kvpair.Key.GetAllocatedSize();
@@ -7656,9 +7657,9 @@ void FPakPlatformFile::OptimizeMemoryUsageForMountedPaks()
 #endif
 
 #if !UE_BUILD_SHIPPING
-	uint32 DirectoryHashSize = 0;
-	uint32 PathHashSize = 0;
-	uint32 EntriesSize = 0;
+	uint64 DirectoryHashSize = 0;
+	uint64 PathHashSize = 0;
+	uint64 EntriesSize = 0;
 
 	for (auto& Pak : Paks)
 	{
diff --git a/Engine/Source/Runtime/PakFile/Private/IoDispatcherFileBackend.cpp b/Engine/Source/Runtime/PakFile/Private/IoDispatcherFileBackend.cpp
--- a/Engine/Source/Runtime/PakFile/Private/IoDispatcherFileBackend.cpp
+++ b/Engine/Source/Runtime/PakFile/Private/IoDispatcherFileBackend.cpp
@@ -884,7 +884,7 @@ TIoStatusOr<FIoContainerHeader> FFileIoStoreReader::ReadContainerHeader() const
 	// Check for flag - compressed containers still have CompressionBlockSize != 0 and CompressionMethod "None".
 	if (EnumHasAnyFlags(ContainerFile.ContainerFlags, EIoContainerFlags::Compressed)) 
 	{
-		FileCache_PostIoStoreCompressionBlockSize(CompressionBlockSize, Partition.FilePath);
+		FileCache_PostIoStoreCompressionBlockSize(IntCastChecked<int32>(CompressionBlockSize), Partition.FilePath);
 	}
 #endif
 
diff --git a/Engine/Source/Runtime/PakFile/Private/SignedArchiveReader.cpp b/Engine/Source/Runtime/PakFile/Private/SignedArchiveReader.cpp
--- a/Engine/Source/Runtime/PakFile/Private/SignedArchiveReader.cpp
+++ b/Engine/Source/Runtime/PakFile/Private/SignedArchiveReader.cpp
@@ -375,15 +375,15 @@ int64 FSignedArchiveReader::PrecacheChunks(TArray<FSignedArchiveReader::FReadInf
 	// Request all the chunks that are needed to complete this read
 	int64 DataOffset = 0;
 	int64 DestOffset = 0;
-	const int32 FirstChunkIndex = CalculateChunkIndex(PakOffset);
-	const int32 LastChunkIndex = CalculateChunkIndex(PakOffset + Length - 1);
+	const int64 FirstChunkIndex = CalculateChunkIndex(PakOffset);
+	const int64 LastChunkIndex = CalculateChunkIndex(PakOffset + Length - 1);
 	const int64 NumChunks = LastChunkIndex - FirstChunkIndex + 1;
 	int64 ChunkStartOffset = 0;
 	int64 RemainingLength = Length;
 	int64 ArchiveOffset = PakOffset;
 
-	Chunks.Empty(NumChunks);
-	for (int32 ChunkIndex = FirstChunkIndex; ChunkIndex <= LastChunkIndex; ++ChunkIndex)
+	Chunks.Empty(IntCastChecked<int32>(NumChunks));
+	for (int32 ChunkIndex = IntCastChecked<int32>(FirstChunkIndex); ChunkIndex <= LastChunkIndex; ++ChunkIndex)
 	{
 		ChunkStartOffset = RemainingLength > 0 ? CalculateChunkOffset(ArchiveOffset, DataOffset) : CalculateChunkOffsetFromIndex(ChunkIndex);
 		int64 SizeToReadFromBuffer = RemainingLength;
@@ -429,8 +429,8 @@ void FSignedArchiveReader::Serialize(void* Data, int64 Length)
 	
 	// First make sure the chunks we're going to read are actually cached.
 	TArray<FReadInfo> QueuedChunks;
-	int64 ChunksToRead = PrecacheChunks(QueuedChunks, Length, ChunkReadEvent);
-	int64 FirstPrecacheChunkIndex = ChunksToRead;
+	int32 ChunksToRead = IntCastChecked<int32>(PrecacheChunks(QueuedChunks, Length, ChunkReadEvent));
+	int32 FirstPrecacheChunkIndex = ChunksToRead;
 
 	// If we aren't multithreaded then flush the signature checking now so there will be some data ready
 	// for us in the loop
diff --git a/Engine/Source/Runtime/PakFile/Private/SignedArchiveReader.h b/Engine/Source/Runtime/PakFile/Private/SignedArchiveReader.h
--- a/Engine/Source/Runtime/PakFile/Private/SignedArchiveReader.h
+++ b/Engine/Source/Runtime/PakFile/Private/SignedArchiveReader.h
@@ -230,7 +230,7 @@ class FSignedArchiveReader : public FArchive
 	};
 
 	/** Number of chunks in the archive */
-	int32 ChunkCount;	
+	int64 ChunkCount;	
 	/** Reader archive */
 	FArchive* PakReader;
 	/** Size of the archive on disk */
@@ -247,9 +247,9 @@ class FSignedArchiveReader : public FArchive
 	/** 
 	 * Calculate index of a chunk that contains the specified offset 
 	 */
-	FORCEINLINE int32 CalculateChunkIndex(int64 ReadOffset) const
+	FORCEINLINE int64 CalculateChunkIndex(int64 ReadOffset) const
 	{
-		return (int32)(ReadOffset / FPakInfo::MaxChunkDataSize);
+		return (ReadOffset / FPakInfo::MaxChunkDataSize);
 	}	
 
 	/** 
@@ -269,7 +269,7 @@ class FSignedArchiveReader : public FArchive
 	 */
 	FORCEINLINE int64 CalculateChunkOffset(int64 ReadOffset, int64& OutDataOffset) const
 	{
-		const int32 ChunkIndex = CalculateChunkIndex(ReadOffset);
+		const int64 ChunkIndex = CalculateChunkIndex(ReadOffset);
 		OutDataOffset = ReadOffset;
 		return CalculateChunkOffsetFromIndex(ChunkIndex);
 	}
diff --git a/Engine/Source/Runtime/PakFile/Public/IPlatformFilePak.h b/Engine/Source/Runtime/PakFile/Public/IPlatformFilePak.h
--- a/Engine/Source/Runtime/PakFile/Public/IPlatformFilePak.h
+++ b/Engine/Source/Runtime/PakFile/Public/IPlatformFilePak.h
@@ -318,7 +318,7 @@ struct FPakInfo
 		return CompressionMethods.Add(CompressionMethod);
 	}
 
-	FName GetCompressionMethod(uint8 Index) const
+	FName GetCompressionMethod(uint32 Index) const
 	{
 		return CompressionMethods[Index];
 	}
