PakFilePrecacher: Fix deadlock by avoiding calling the body of ClearOldBlockTasks while holding CachedFilesScopeLock.

ClearOldBlockTasks can BusyWait, which can call FAsyncIOCPUWorkTask::DoTask, which can try to take a FPakAsyncReadFileHandle's CriticalSection. FPakAsyncReadFileHandle's CriticalSection must be entered first when holding both CriticalSection and CachedFilesScopeLock to prevent deadlocks.
#rnx

diff --git a/Engine/Source/Runtime/PakFile/Private/IPlatformFilePak.cpp b/Engine/Source/Runtime/PakFile/Private/IPlatformFilePak.cpp
--- a/Engine/Source/Runtime/PakFile/Private/IPlatformFilePak.cpp
+++ b/Engine/Source/Runtime/PakFile/Private/IPlatformFilePak.cpp
@@ -2875,6 +2875,11 @@ public:
 		{
 			return false;
 		}
+		// Use NotifyRecursion to turn off maintenance functions like ClearOldBlockTasks that can BusyWait
+		// (and thereby reenter this class and take locks in the wrong order) while we are holding the lock.
+		++NotifyRecursion;
+		ON_SCOPE_EXIT{ --NotifyRecursion; };
+
 		uint16 PakIndex = *PakIndexPtr;
 		FPakData& Pak = CachedPakData[PakIndex];
 		check(Pak.Name == File && Pak.TotalSize == PakFileSize && Pak.Handle);
