IoDispatcher: Keep priority value on resolved requests to avoid having to go out and read from the dispatcher requests

diff --git a/Engine/Source/Runtime/PakFile/Internal/IoDispatcherFileBackendTypes.h b/Engine/Source/Runtime/PakFile/Internal/IoDispatcherFileBackendTypes.h
--- a/Engine/Source/Runtime/PakFile/Internal/IoDispatcherFileBackendTypes.h
+++ b/Engine/Source/Runtime/PakFile/Internal/IoDispatcherFileBackendTypes.h
@@ -796,13 +796,15 @@ public:
 		FIoRequestImpl& InDispatcherRequest,
 		FFileIoStoreContainerFile* InContainerFile,
 		uint64 InResolvedOffset,
-		uint64 InResolvedSize)
+		uint64 InResolvedSize,
+		int32 InPriority)
 	{
 		return ResolvedRequestAllocator.Construct(
 			InDispatcherRequest,
 			InContainerFile,
 			InResolvedOffset,
-			InResolvedSize);
+			InResolvedSize,
+			InPriority);
 	}
 
 	void Free(FFileIoStoreResolvedRequest* ResolvedRequest)
@@ -862,7 +864,8 @@ public:
 		FIoRequestImpl& InDispatcherRequest,
 		FFileIoStoreContainerFile* InContainerFile,
 		uint64 InResolvedOffset,
-		uint64 InResolvedSize);
+		uint64 InResolvedSize,
+		int32 InPriority);
 
 	const FFileIoStoreContainerFile* GetContainerFile() const
 	{
@@ -886,8 +889,7 @@ public:
 
 	int32 GetPriority() const
 	{
-		check(DispatcherRequest);
-		return DispatcherRequest->Priority;
+		return Priority;
 	}
 
 	bool HasBuffer() const
@@ -917,6 +919,7 @@ private:
 	FFileIoStoreReadRequestLink* ReadRequestsTail = nullptr;
 	const uint64 ResolvedOffset;
 	const uint64 ResolvedSize;
+	int32 Priority = 0;
 	uint32 UnfinishedReadsCount = 0;
 	bool bFailed = false;
 	bool bCancelled = false;
diff --git a/Engine/Source/Runtime/PakFile/Private/IoDispatcherFileBackend.cpp b/Engine/Source/Runtime/PakFile/Private/IoDispatcherFileBackend.cpp
--- a/Engine/Source/Runtime/PakFile/Private/IoDispatcherFileBackend.cpp
+++ b/Engine/Source/Runtime/PakFile/Private/IoDispatcherFileBackend.cpp
@@ -958,11 +958,13 @@ FFileIoStoreResolvedRequest::FFileIoStoreResolvedRequest(
 	FIoRequestImpl& InDispatcherRequest,
 	FFileIoStoreContainerFile* InContainerFile,
 	uint64 InResolvedOffset,
-	uint64 InResolvedSize)
+	uint64 InResolvedSize,
+	int32 InPriority)
 	: DispatcherRequest(&InDispatcherRequest)
 	, ContainerFile(InContainerFile)
 	, ResolvedOffset(InResolvedOffset)
 	, ResolvedSize(InResolvedSize)
+	, Priority(InPriority)
 {
 
 }
@@ -1385,7 +1387,8 @@ bool FFileIoStore::Resolve(FIoRequestImpl* Request)
 				*Request,
 				Reader->GetContainerFile(),
 				ResolvedOffset,
-				ResolvedSize);
+				ResolvedSize,
+				Request->Priority);
 			Request->BackendData = ResolvedRequest;
 
 			if (ResolvedSize > 0)
@@ -1444,6 +1447,7 @@ void FFileIoStore::UpdatePriorityForIoRequest(FIoRequestImpl* Request)
 	if (Request->BackendData)
 	{
 		FFileIoStoreResolvedRequest* ResolvedRequest = static_cast<FFileIoStoreResolvedRequest*>(Request->BackendData);
+		ResolvedRequest->Priority = Request->Priority;
 		RequestTracker.UpdatePriorityForIoRequest(*ResolvedRequest);
 	}
 }
