IoDispatcherFileBackend: Ensure that all decompression tasks have finished when shutting down
#rnx

diff --git a/Engine/Source/Runtime/PakFile/Private/IoDispatcherFileBackend.cpp b/Engine/Source/Runtime/PakFile/Private/IoDispatcherFileBackend.cpp
--- a/Engine/Source/Runtime/PakFile/Private/IoDispatcherFileBackend.cpp
+++ b/Engine/Source/Runtime/PakFile/Private/IoDispatcherFileBackend.cpp
@@ -1245,12 +1245,13 @@ void FFileIoStore::Initialize(TSharedRef<const FIoDispatcherBackendContext> InCo
 		&Stats
 	});
 
-	uint64 DecompressionContextCount = uint64(GIoDispatcherDecompressionWorkerCount > 0 ? GIoDispatcherDecompressionWorkerCount : 4);
-	for (uint64 ContextIndex = 0; ContextIndex < DecompressionContextCount; ++ContextIndex)
+	int32 DecompressionContextCount = int32(GIoDispatcherDecompressionWorkerCount > 0 ? GIoDispatcherDecompressionWorkerCount : 4);
+	CompressionContexts.SetNum(DecompressionContextCount);
+	for (TUniquePtr<FFileIoStoreCompressionContext>& CompressionContext : CompressionContexts)
 	{
-		FFileIoStoreCompressionContext* Context = new FFileIoStoreCompressionContext();
-		Context->Next = FirstFreeCompressionContext;
-		FirstFreeCompressionContext = Context;
+		CompressionContext = MakeUnique<FFileIoStoreCompressionContext>();
+		CompressionContext->Next = FirstFreeCompressionContext;
+		FirstFreeCompressionContext = CompressionContext.Get();
 	}
 
 	Thread = FRunnableThread::Create(this, TEXT("IoService"), 0, TPri_AboveNormal);
@@ -1827,7 +1828,7 @@ FIoRequestImpl* FFileIoStore::GetCompletedRequests()
 		const bool bScatterAsync = bIsMultithreaded && (!BlockToDecompress->CompressionMethod.IsNone() || BlockToDecompress->EncryptionKey.IsValid() || BlockToDecompress->SignatureHash);
 		if (bScatterAsync)
 		{
-			TGraphTask<FDecompressAsyncTask>::CreateTask().ConstructAndDispatchWhenReady(*this, BlockToDecompress);
+			BlockToDecompress->CompressionContext->Task = TGraphTask<FDecompressAsyncTask>::CreateTask().ConstructAndDispatchWhenReady(*this, BlockToDecompress);
 		}
 		else
 		{
@@ -2058,6 +2059,15 @@ uint32 FFileIoStore::Run()
 			PlatformImpl->ServiceWait();
 		}
 	}
+
+	for (TUniquePtr<FFileIoStoreCompressionContext>& CompressionContext : CompressionContexts)
+	{
+		if (CompressionContext->Task)
+		{
+			CompressionContext->Task->Wait();
+		}
+	}
+
 	return 0;
 }
 
diff --git a/Engine/Source/Runtime/PakFile/Private/IoDispatcherFileBackend.h b/Engine/Source/Runtime/PakFile/Private/IoDispatcherFileBackend.h
--- a/Engine/Source/Runtime/PakFile/Private/IoDispatcherFileBackend.h
+++ b/Engine/Source/Runtime/PakFile/Private/IoDispatcherFileBackend.h
@@ -20,6 +20,7 @@ struct FFileIoStoreCompressionContext
 	FFileIoStoreCompressionContext* Next = nullptr;
 	uint64 UncompressedBufferSize = 0;
 	uint8* UncompressedBuffer = nullptr;
+	FGraphEventRef Task;
 };
 
 class FFileIoStoreReader
@@ -144,7 +145,7 @@ private:
 
 		FORCEINLINE static ESubsequentsMode::Type GetSubsequentsMode()
 		{
-			return ESubsequentsMode::FireAndForget;
+			return ESubsequentsMode::TrackSubsequents;
 		}
 
 		void DoTask(ENamedThreads::Type CurrentThread, const FGraphEventRef& MyCompletionGraphEvent)
@@ -181,6 +182,7 @@ private:
 	TAtomic<bool> bStopRequested{ false };
 	mutable FRWLock IoStoreReadersLock;
 	TArray<TUniquePtr<FFileIoStoreReader>> IoStoreReaders;
+	TArray<TUniquePtr<FFileIoStoreCompressionContext>> CompressionContexts;
 	FFileIoStoreCompressionContext* FirstFreeCompressionContext = nullptr;
 	FFileIoStoreCompressedBlock* ReadyForDecompressionHead = nullptr;
 	FFileIoStoreCompressedBlock* ReadyForDecompressionTail = nullptr;
