Prevent memory mapping compressed or encrypted memory with IoStore

diff --git a/Engine/Source/Runtime/PakFile/Private/IoDispatcherFileBackend.cpp b/Engine/Source/Runtime/PakFile/Private/IoDispatcherFileBackend.cpp
--- a/Engine/Source/Runtime/PakFile/Private/IoDispatcherFileBackend.cpp
+++ b/Engine/Source/Runtime/PakFile/Private/IoDispatcherFileBackend.cpp
@@ -111,14 +111,13 @@ public:
 		: IMappedFileHandle(InSize)
 		, SharedMappedFileHandle(InSharedMappedFileHandle)
 	{
-		check(InSharedMappedFileHandle != nullptr);
 	}
 
 	virtual ~FMappedFileProxy() { }
 
 	virtual IMappedFileRegion* MapRegion(int64 Offset = 0, int64 BytesToMap = MAX_int64, bool bPreloadHint = false) override
 	{
-		return SharedMappedFileHandle->MapRegion(Offset, BytesToMap, bPreloadHint);
+		return SharedMappedFileHandle != nullptr ? SharedMappedFileHandle->MapRegion(Offset, BytesToMap, bPreloadHint) : nullptr;
 	}
 private:
 	IMappedFileHandle* SharedMappedFileHandle;
@@ -846,7 +845,9 @@ IMappedFileHandle* FFileIoStoreReader::GetMappedContainerFileHandle(uint64 TocOf
 	check(!bClosed);
 	int32 PartitionIndex = int32(TocOffset / ContainerFile.PartitionSize);
 	FFileIoStoreContainerFilePartition& Partition = ContainerFile.Partitions[PartitionIndex];
-	if (!Partition.MappedFileHandle)
+	if (!Partition.MappedFileHandle &&
+		// Can't map encrypted files, compression should be disabled for bulk data when memory mapping is required
+		!EnumHasAnyFlags(ContainerFile.ContainerFlags, EIoContainerFlags::Encrypted))
 	{
 		IPlatformFile& Ipf = FPlatformFileManager::Get().GetPlatformFile();
 		Partition.MappedFileHandle.Reset(Ipf.OpenMapped(*Partition.FilePath));
@@ -1860,8 +1861,6 @@ TIoStatusOr<FIoMappedRegion> FFileIoStore::OpenMapped(const FIoChunkId& ChunkId,
 		return FIoStatus(EIoErrorCode::InvalidParameter, TEXT("Invalid read options"));
 	}
 
-	IPlatformFile& Ipf = FPlatformFileManager::Get().GetPlatformFile();
-
 	FReadScopeLock _(IoStoreReadersLock);
 	for (TUniquePtr<FFileIoStoreReader>& Reader : IoStoreReaders)
 	{
