PakFiles: If FChunkCacheWorker's construction fails and the FChunkCacheWorker is invalid, fail the entire Pak by returning nullptr from CreatePakReader rather than allowing creation of an FSignedArchiveReader that will crash in FChunkCacheWorker::CheckSignature.
#rnx

diff --git a/Engine/Source/Runtime/PakFile/Private/IPlatformFilePak.cpp b/Engine/Source/Runtime/PakFile/Private/IPlatformFilePak.cpp
--- a/Engine/Source/Runtime/PakFile/Private/IPlatformFilePak.cpp
+++ b/Engine/Source/Runtime/PakFile/Private/IPlatformFilePak.cpp
@@ -5442,7 +5442,7 @@ FArchive* FPakFile::CreatePakReader(IPlatformFile* LowerLevel, const TCHAR* File
 			Decryptor = MakeUnique<FChunkCacheWorker>(MoveTemp(DecryptorReader), Filename);
 		}
 
-		if (!Decryptor.IsValid())
+		if (!Decryptor.IsValid() || !Decryptor->IsValid())
 		{
 			return nullptr;
 		}
