Validating index loading for upacks.

diff --git a/Engine/Source/Runtime/PakFile/Private/IPlatformFilePak.cpp b/Engine/Source/Runtime/PakFile/Private/IPlatformFilePak.cpp
--- a/Engine/Source/Runtime/PakFile/Private/IPlatformFilePak.cpp
+++ b/Engine/Source/Runtime/PakFile/Private/IPlatformFilePak.cpp
@@ -2,6 +2,7 @@
 
 #include "IPlatformFilePak.h"
 #include "HAL/FileManager.h"
+#include "Math/GuardedInt.h"
 #include "Misc/CoreMisc.h"
 #include "Misc/CommandLine.h"
 #include "Async/AsyncWork.h"
@@ -5584,14 +5585,19 @@ bool FPakFile::LoadIndexInternal(FArchive& Reader)
 	bWillPruneDirectoryIndex = false;
 #endif
 
-	if (CachedTotalSize < (Info.IndexOffset + Info.IndexSize))
+	FGuardedInt64 IndexEndPosition = FGuardedInt64(Info.IndexOffset) + Info.IndexSize;
+	if (Info.IndexOffset < 0 || 
+		Info.IndexSize < 0 ||
+		IndexEndPosition.InvalidOrGreaterThan(CachedTotalSize) ||
+		IntFitsIn<int32>(Info.IndexSize) == false)
 	{
 		UE_LOG(LogPakFile, Fatal, TEXT("Corrupted index offset in pak file."));
 		return false;
 	}
+
 	TArray<uint8> PrimaryIndexData;
 	Reader.Seek(Info.IndexOffset);
-	PrimaryIndexData.SetNum(IntCastChecked<int32>(Info.IndexSize));
+	PrimaryIndexData.SetNum((int32)(Info.IndexSize));
 	{
 		SCOPED_BOOT_TIMING("PakFile_LoadPrimaryIndex");
 		Reader.Serialize(PrimaryIndexData.GetData(), Info.IndexSize);
@@ -5605,9 +5611,9 @@ bool FPakFile::LoadIndexInternal(FArchive& Reader)
 			UE_LOG(LogPakFile, Log, TEXT("Corrupt pak PrimaryIndex detected!"));
 			UE_LOG(LogPakFile, Log, TEXT(" Filename: %s"), *PakFilename);
 			UE_LOG(LogPakFile, Log, TEXT(" Encrypted: %d"), Info.bEncryptedIndex);
-			UE_LOG(LogPakFile, Log, TEXT(" Total Size: %d"), Reader.TotalSize());
-			UE_LOG(LogPakFile, Log, TEXT(" Index Offset: %d"), Info.IndexOffset);
-			UE_LOG(LogPakFile, Log, TEXT(" Index Size: %d"), Info.IndexSize);
+			UE_LOG(LogPakFile, Log, TEXT(" Total Size: %lld"), Reader.TotalSize());
+			UE_LOG(LogPakFile, Log, TEXT(" Index Offset: %lld"), Info.IndexOffset);
+			UE_LOG(LogPakFile, Log, TEXT(" Index Size: %lld"), Info.IndexSize);
 			UE_LOG(LogPakFile, Log, TEXT(" Stored Index Hash: %s"), *Info.IndexHash.ToString());
 			UE_LOG(LogPakFile, Log, TEXT(" Computed Index Hash: %s"), *ComputedHash.ToString());
 			return false;
@@ -5619,8 +5625,22 @@ bool FPakFile::LoadIndexInternal(FArchive& Reader)
 	// Read the scalar data (mount point, numentries, etc) and all entries.
 	NumEntries = 0;
 	PrimaryIndexReader << MountPoint;
+	// We are just deserializing a string, which could get however long. Since we know it's a path
+	// and paths are bound to file system rules, we know it can't get absurdly long (e.g. windows is _at best_ 32k)
+	// we just sanity check it to prevent operating on massive buffers and risking overflows.
+	if (MountPoint.Len() > 65535)
+	{
+		UE_LOG(LogPakFile, Error, TEXT("Corrupt pak index data: MountPoint path is longer than 65k"));
+		return false;
+	}
+
 	MakeDirectoryFromPath(MountPoint);
 	PrimaryIndexReader << NumEntries;
+	if (NumEntries < 0)
+	{
+		UE_LOG(LogPakFile, Error, TEXT("Corrupt pak index data: Negative entries count in pak file."));
+		return false;
+	}
 	PrimaryIndexReader << PathHashSeed;
 
 	bool bReaderHasPathHashIndex = false;
@@ -5715,18 +5735,22 @@ bool FPakFile::LoadIndexInternal(FArchive& Reader)
 	FMemoryReader PathHashIndexReader(PathHashIndexData);
 	if (bWillUsePathHashIndex)
 	{
-		if (PathHashIndexOffset < 0 || CachedTotalSize < (PathHashIndexOffset + PathHashIndexSize))
+		FGuardedInt64 PathHashIndexEndPosition = FGuardedInt64(PathHashIndexOffset) + PathHashIndexSize;
+		if (PathHashIndexOffset < 0 || 
+			PathHashIndexSize < 0 ||
+			PathHashIndexEndPosition.InvalidOrGreaterThan(CachedTotalSize) || 
+			IntFitsIn<int32>(PathHashIndexSize) == false)
 		{
 			// Should not be possible for these values (which came from the PrimaryIndex) to be invalid, since we verified the index hash of the PrimaryIndex
 			UE_LOG(LogPakFile, Log, TEXT("Corrupt pak PrimaryIndex detected!"));
 			UE_LOG(LogPakFile, Log, TEXT(" Filename: %s"), *PakFilename);
-			UE_LOG(LogPakFile, Log, TEXT(" Total Size: %d"), Reader.TotalSize());
-			UE_LOG(LogPakFile, Log, TEXT(" PathHashIndexOffset : %d"), PathHashIndexOffset);
-			UE_LOG(LogPakFile, Log, TEXT(" PathHashIndexSize: %d"), PathHashIndexSize);
+			UE_LOG(LogPakFile, Log, TEXT(" Total Size: %lld"), Reader.TotalSize());
+			UE_LOG(LogPakFile, Log, TEXT(" PathHashIndexOffset : %lld"), PathHashIndexOffset);
+			UE_LOG(LogPakFile, Log, TEXT(" PathHashIndexSize: %lld"), PathHashIndexSize);
 			return false;
 		}
 		Reader.Seek(PathHashIndexOffset);
-		PathHashIndexData.SetNum(IntCastChecked<int32>(PathHashIndexSize));
+		PathHashIndexData.SetNum((int32)(PathHashIndexSize));
 		{
 			SCOPED_BOOT_TIMING("PakFile_LoadPathHashIndex");
 			Reader.Serialize(PathHashIndexData.GetData(), PathHashIndexSize);
@@ -5739,9 +5763,9 @@ bool FPakFile::LoadIndexInternal(FArchive& Reader)
 				UE_LOG(LogPakFile, Log, TEXT("Corrupt pak PathHashIndex detected!"));
 				UE_LOG(LogPakFile, Log, TEXT(" Filename: %s"), *PakFilename);
 				UE_LOG(LogPakFile, Log, TEXT(" Encrypted: %d"), Info.bEncryptedIndex);
-				UE_LOG(LogPakFile, Log, TEXT(" Total Size: %d"), Reader.TotalSize());
-				UE_LOG(LogPakFile, Log, TEXT(" Index Offset: %d"), FullDirectoryIndexOffset);
-				UE_LOG(LogPakFile, Log, TEXT(" Index Size: %d"), FullDirectoryIndexSize);
+				UE_LOG(LogPakFile, Log, TEXT(" Total Size: %lld"), Reader.TotalSize());
+				UE_LOG(LogPakFile, Log, TEXT(" Index Offset: %lld"), FullDirectoryIndexOffset);
+				UE_LOG(LogPakFile, Log, TEXT(" Index Size: %lld"), FullDirectoryIndexSize);
 				UE_LOG(LogPakFile, Log, TEXT(" Stored Index Hash: %s"), *PathHashIndexHash.ToString());
 				UE_LOG(LogPakFile, Log, TEXT(" Computed Index Hash: %s"), *ComputedHash.ToString());
 				return false;
@@ -5770,20 +5794,23 @@ bool FPakFile::LoadIndexInternal(FArchive& Reader)
 	}
 	else
 	{
-		if (CachedTotalSize < (FullDirectoryIndexOffset  + FullDirectoryIndexSize) ||
-			FullDirectoryIndexOffset  < 0)
+		FGuardedInt64 FullDirectoryIndexEndPosition = FGuardedInt64(FullDirectoryIndexOffset) + FullDirectoryIndexSize;
+		if (FullDirectoryIndexOffset  < 0 || 
+			FullDirectoryIndexSize < 0 ||
+			FullDirectoryIndexEndPosition.InvalidOrGreaterThan(CachedTotalSize) || 
+			IntFitsIn<int32>(FullDirectoryIndexSize) == false)
 		{
 			// Should not be possible for these values (which came from the PrimaryIndex) to be invalid, since we verified the index hash of the PrimaryIndex
 			UE_LOG(LogPakFile, Log, TEXT("Corrupt pak PrimaryIndex detected!"));
 			UE_LOG(LogPakFile, Log, TEXT(" Filename: %s"), *PakFilename);
-			UE_LOG(LogPakFile, Log, TEXT(" Total Size: %d"), Reader.TotalSize());
-			UE_LOG(LogPakFile, Log, TEXT(" FullDirectoryIndexOffset : %d"), FullDirectoryIndexOffset );
-			UE_LOG(LogPakFile, Log, TEXT(" FullDirectoryIndexSize: %d"), FullDirectoryIndexSize);
+			UE_LOG(LogPakFile, Log, TEXT(" Total Size: %lld"), Reader.TotalSize());
+			UE_LOG(LogPakFile, Log, TEXT(" FullDirectoryIndexOffset : %lld"), FullDirectoryIndexOffset );
+			UE_LOG(LogPakFile, Log, TEXT(" FullDirectoryIndexSize: %lld"), FullDirectoryIndexSize);
 			return false;
 		}
 		TArray<uint8> FullDirectoryIndexData;
 		Reader.Seek(FullDirectoryIndexOffset );
-		FullDirectoryIndexData.SetNum(IntCastChecked<int32>(FullDirectoryIndexSize));
+		FullDirectoryIndexData.SetNum((int32)(FullDirectoryIndexSize));
 		{
 			SCOPED_BOOT_TIMING("PakFile_LoadDirectoryIndex");
 			Reader.Serialize(FullDirectoryIndexData.GetData(), FullDirectoryIndexSize);
@@ -5796,9 +5823,9 @@ bool FPakFile::LoadIndexInternal(FArchive& Reader)
 				UE_LOG(LogPakFile, Log, TEXT("Corrupt pak FullDirectoryIndex detected!"));
 				UE_LOG(LogPakFile, Log, TEXT(" Filename: %s"), *PakFilename);
 				UE_LOG(LogPakFile, Log, TEXT(" Encrypted: %d"), Info.bEncryptedIndex);
-				UE_LOG(LogPakFile, Log, TEXT(" Total Size: %d"), Reader.TotalSize());
-				UE_LOG(LogPakFile, Log, TEXT(" Index Offset: %d"), FullDirectoryIndexOffset);
-				UE_LOG(LogPakFile, Log, TEXT(" Index Size: %d"), FullDirectoryIndexSize);
+				UE_LOG(LogPakFile, Log, TEXT(" Total Size: %lld"), Reader.TotalSize());
+				UE_LOG(LogPakFile, Log, TEXT(" Index Offset: %lld"), FullDirectoryIndexOffset);
+				UE_LOG(LogPakFile, Log, TEXT(" Index Size: %lld"), FullDirectoryIndexSize);
 				UE_LOG(LogPakFile, Log, TEXT(" Stored Index Hash: %s"), *FullDirectoryIndexHash.ToString());
 				UE_LOG(LogPakFile, Log, TEXT(" Computed Index Hash: %s"), *ComputedHash.ToString());
 				return false;
@@ -5849,16 +5876,20 @@ bool FPakFile::LoadLegacyIndex(FArchive& Reader)
 #endif
 
 	// Load index into memory first.
-
-	TArray<uint8> IndexData;
-	if (CachedTotalSize < (Info.IndexOffset + Info.IndexSize))
+	FGuardedInt64 IndexEndPosition = FGuardedInt64(Info.IndexOffset) + Info.IndexSize;
+	if (Info.IndexSize < 0 ||
+		Info.IndexOffset < 0 ||
+		IndexEndPosition.InvalidOrGreaterThan(CachedTotalSize) ||
+		IntFitsIn<int32>(Info.IndexSize) == false)
 	{
-		UE_LOG(LogPakFile, Fatal, TEXT("Corrupted index offset in pak file."));
+		UE_LOG(LogPakFile, Fatal, TEXT("Corrupted index offset/size in pak file."));
 		return false;
 	}
 
+	TArray<uint8> IndexData;
+	IndexData.SetNum((int32)(Info.IndexSize));
+
 	Reader.Seek(Info.IndexOffset);
-	IndexData.SetNum(IntCastChecked<int32>(Info.IndexSize));
 	Reader.Serialize(IndexData.GetData(), Info.IndexSize);
 
 	FSHAHash ComputedHash;
@@ -5867,9 +5898,9 @@ bool FPakFile::LoadLegacyIndex(FArchive& Reader)
 		UE_LOG(LogPakFile, Log, TEXT("Corrupt pak index detected!"));
 		UE_LOG(LogPakFile, Log, TEXT(" Filename: %s"), *PakFilename);
 		UE_LOG(LogPakFile, Log, TEXT(" Encrypted: %d"), Info.bEncryptedIndex);
-		UE_LOG(LogPakFile, Log, TEXT(" Total Size: %d"), Reader.TotalSize());
-		UE_LOG(LogPakFile, Log, TEXT(" Index Offset: %d"), Info.IndexOffset);
-		UE_LOG(LogPakFile, Log, TEXT(" Index Size: %d"), Info.IndexSize);
+		UE_LOG(LogPakFile, Log, TEXT(" Total Size: %lld"), Reader.TotalSize());
+		UE_LOG(LogPakFile, Log, TEXT(" Index Offset: %lld"), Info.IndexOffset);
+		UE_LOG(LogPakFile, Log, TEXT(" Index Size: %lld"), Info.IndexSize);
 		UE_LOG(LogPakFile, Log, TEXT(" Stored Index Hash: %s"), *Info.IndexHash.ToString());
 		UE_LOG(LogPakFile, Log, TEXT(" Computed Index Hash: %s"), *ComputedHash.ToString());
 		return false;
@@ -5881,8 +5912,23 @@ bool FPakFile::LoadLegacyIndex(FArchive& Reader)
 	// Read the default mount point and all entries.
 	NumEntries = 0;
 	IndexReader << MountPoint;
+
+	// We are just deserializing a string, which could get however long. Since we know it's a path
+	// and paths are bound to file system rules, we know it can't get absurdly long (e.g. windows is _at best_ 32k)
+	// we just sanity check it to prevent operating on massive buffers and risking overflows.
+	if (MountPoint.Len() > 65535)
+	{
+		UE_LOG(LogPakFile, Error, TEXT("Corrupt pak index data: MountPoint path is longer than 65k"));
+		return false;
+	}
 	IndexReader << NumEntries;
 
+	if (NumEntries < 0)
+	{
+		UE_LOG(LogPakFile, Error, TEXT("Corrupt pak index data: NumEntries is negative"));
+		return false;
+	}
+
 	MakeDirectoryFromPath(MountPoint);
 
 	FPakEntryPair PakEntryPair;
diff --git a/Engine/Source/Runtime/PakFile/Public/IPlatformFilePak.h b/Engine/Source/Runtime/PakFile/Public/IPlatformFilePak.h
--- a/Engine/Source/Runtime/PakFile/Public/IPlatformFilePak.h
+++ b/Engine/Source/Runtime/PakFile/Public/IPlatformFilePak.h
@@ -280,6 +280,8 @@ struct FPakInfo
 					ANSICHAR* MethodString = &Methods[Index * CompressionMethodNameLen];
 					if (MethodString[0] != 0)
 					{
+						// Make sure if the file has garbage we don't read off in to lala land.
+						MethodString[CompressionMethodNameLen-1] = 0;
 						CompressionMethods.Add(FName(MethodString));
 					}
 				}
