AsyncIO: Fixed incorrect handling of the STAT_AsyncFileMemory stat.

I've reviewed all code using async requests and we had multiple places incorrectly handling the stat: some placed failed to decrease the counter, while others would decrease it when they shouldn't. The main cause was the fact that when we create an async request and pass nullptr as the target memory, the request allocates the memory itself and increases the STAT_AsyncFileMemory. However, when we call GetReadResults() on such requests, the responsibility is on the caller to balance the calls and decrease the stat, which is not very obvious and leads to the mentioned bugs. On top of that, when we call GetReadResults() we may not have the knowledge as to whether the returned memory was allocated by the request or not (which may affect whether we should touch the stat at all).

Resolved by adding IAsyncReadRequest::ReleaseMemoryOwnershipImpl member, which may be used by request implementations to do the proper clean up (balance the stat). This way requests take the full responsibility for their changes and users of the requests don't need to know their implementation details.

diff --git a/Engine/Plugins/Runtime/GeometryCache/Source/GeometryCache/Private/StreamingGeometryCacheData.cpp b/Engine/Plugins/Runtime/GeometryCache/Source/GeometryCache/Private/StreamingGeometryCacheData.cpp
--- a/Engine/Plugins/Runtime/GeometryCache/Source/GeometryCache/Private/StreamingGeometryCacheData.cpp
+++ b/Engine/Plugins/Runtime/GeometryCache/Source/GeometryCache/Private/StreamingGeometryCacheData.cpp
@@ -355,7 +355,6 @@ void FStreamingGeometryCacheData::ProcessCompletedChunks()
 				DEC_DWORD_STAT(STAT_OutstandingRequests);
 				INC_MEMORY_STAT_BY(STAT_ChunkDataResident, Chunk->DataSize);
 				INC_MEMORY_STAT_BY(STAT_ChunkDataStreamed, Chunk->DataSize);
-				DEC_MEMORY_STAT_BY(STAT_AsyncFileMemory, Chunk->DataSize);
 				IGeometryCacheStreamingManager::Get().IoBandwidth.Add(Chunk->DataSize);
 			}
 			else
diff --git a/Engine/Source/Editor/UnrealEd/Private/EditorDomain/EditorDomainArchive.cpp b/Engine/Source/Editor/UnrealEd/Private/EditorDomain/EditorDomainArchive.cpp
--- a/Engine/Source/Editor/UnrealEd/Private/EditorDomain/EditorDomainArchive.cpp
+++ b/Engine/Source/Editor/UnrealEd/Private/EditorDomain/EditorDomainArchive.cpp
@@ -718,6 +718,10 @@ protected:
 	virtual void CancelImpl() override
 	{
 	}
+
+	virtual void ReleaseMemoryOwnershipImpl() override
+	{
+	}
 };
 
 /** An IAsyncReadRequest that reads from memory that was already populated. */
@@ -756,6 +760,10 @@ protected:
 	virtual void CancelImpl() override
 	{
 	}
+
+	virtual void ReleaseMemoryOwnershipImpl() override
+	{
+	}
 };
 
 
diff --git a/Engine/Source/Runtime/Core/Private/GenericPlatform/GenericPlatformFile.cpp b/Engine/Source/Runtime/Core/Private/GenericPlatform/GenericPlatformFile.cpp
--- a/Engine/Source/Runtime/Core/Private/GenericPlatform/GenericPlatformFile.cpp
+++ b/Engine/Source/Runtime/Core/Private/GenericPlatform/GenericPlatformFile.cpp
@@ -128,6 +128,11 @@ public:
 		}
 		SetComplete();
 	}
+
+protected:
+	virtual void ReleaseMemoryOwnershipImpl() override
+	{
+	}
 };
 
 class FGenericReadRequest : public FGenericBaseRequest
@@ -174,6 +179,12 @@ public:
 		}
 		return nullptr;
 	}
+
+protected:
+	virtual void ReleaseMemoryOwnershipImpl() override
+	{
+		DEC_MEMORY_STAT_BY(STAT_AsyncFileMemory, BytesToRead);
+	}
 };
 
 void FGenericReadRequestWorker::DoWork()
diff --git a/Engine/Source/Runtime/Core/Private/Microsoft/MicrosoftAsyncIO.h b/Engine/Source/Runtime/Core/Private/Microsoft/MicrosoftAsyncIO.h
--- a/Engine/Source/Runtime/Core/Private/Microsoft/MicrosoftAsyncIO.h
+++ b/Engine/Source/Runtime/Core/Private/Microsoft/MicrosoftAsyncIO.h
@@ -327,6 +327,10 @@ public:
 		// no cancel support
 	}
 
+	virtual void ReleaseMemoryOwnershipImpl() override
+	{
+		DEC_MEMORY_STAT_BY(STAT_AsyncFileMemory, BytesToRead);
+	}
 };
 
 void FMicrosoftReadRequestWorker::DoWork()
@@ -351,7 +355,11 @@ public:
 		while (!*(volatile bool*)&bCompleteAndCallbackCalled);
 	}
 
-	virtual void CancelImpl()
+	virtual void CancelImpl() override
+	{
+	}
+
+	virtual void ReleaseMemoryOwnershipImpl() override
 	{
 	}
 };
@@ -372,7 +380,11 @@ public:
 		while (!*(volatile bool*)&bCompleteAndCallbackCalled);
 	}
 
-	virtual void CancelImpl()
+	virtual void CancelImpl() override
+	{
+	}
+
+	virtual void ReleaseMemoryOwnershipImpl() override
 	{
 	}
 };
diff --git a/Engine/Source/Runtime/Core/Public/Async/AsyncFileHandle.h b/Engine/Source/Runtime/Core/Public/Async/AsyncFileHandle.h
--- a/Engine/Source/Runtime/Core/Public/Async/AsyncFileHandle.h
+++ b/Engine/Source/Runtime/Core/Public/Async/AsyncFileHandle.h
@@ -135,6 +135,11 @@ public:
 		}
 		else
 		{
+			if (Memory && !bUserSuppliedMemory)
+			{
+				ReleaseMemoryOwnershipImpl();
+			}
+
 			Memory = nullptr;
 		}
 		return Result;
@@ -151,6 +156,13 @@ protected:
 	/** Cancel the request. This is a non-blocking async call and so does not ensure completion! **/
 	virtual void CancelImpl() = 0;
 
+	/** Transfer ownership of Memory from the async request to the outside caller (called in response to GetReadResults).
+	  * It's only relevant to Read requests, in which case the most common use is to update (decrease) the STAT_AsyncFileMemory
+	  * stat which is typically incremented when async requests allocate Memory.
+	  * It doesn't play any role in Size requests, so it may be left empty for them.
+	  */
+	virtual void ReleaseMemoryOwnershipImpl() = 0;
+
 	void SetDataComplete()
 	{
 		bDataIsReady = true;
diff --git a/Engine/Source/Runtime/Core/Public/Async/AsyncFileHandleNull.h b/Engine/Source/Runtime/Core/Public/Async/AsyncFileHandleNull.h
--- a/Engine/Source/Runtime/Core/Public/Async/AsyncFileHandleNull.h
+++ b/Engine/Source/Runtime/Core/Public/Async/AsyncFileHandleNull.h
@@ -27,6 +27,10 @@ protected:
 	virtual void CancelImpl() override
 	{
 	}
+
+	virtual void ReleaseMemoryOwnershipImpl() override
+	{
+	}
 };
 
 /**
diff --git a/Engine/Source/Runtime/CoreUObject/Private/Serialization/AsyncLoading.cpp b/Engine/Source/Runtime/CoreUObject/Private/Serialization/AsyncLoading.cpp
--- a/Engine/Source/Runtime/CoreUObject/Private/Serialization/AsyncLoading.cpp
+++ b/Engine/Source/Runtime/CoreUObject/Private/Serialization/AsyncLoading.cpp
@@ -7816,7 +7816,6 @@ void FAsyncArchive::ReadCallback(bool bWasCancelled, IAsyncReadRequest* Request)
 				LoadPhase = ELoadPhase::WaitingForHeader;
 			}
 			FMemory::Free(Mem);
-			DEC_MEMORY_STAT_BY(STAT_AsyncFileMemory, FMath::Min<int64>(FMaxPackageSummarySize::Value, FileSize));
 		}
 	}
 	else
@@ -8013,7 +8012,6 @@ void FAsyncArchive::CompleteRead()
 			PrecacheEndPos = ReadRequestOffset + ReadRequestSize;
 			check(ReadRequestSize > 0 && PrecacheStartPos >= 0);
 			INC_MEMORY_STAT_BY(STAT_FAsyncArchiveMem, PrecacheEndPos - PrecacheStartPos);
-			DEC_MEMORY_STAT_BY(STAT_AsyncFileMemory, ReadRequestSize);
 #if USE_DETAILED_FASYNCARCHIVE_MEMORY_TRACKING
 			GAsyncArchiveMemTracker.Allocate(FileName, PrecacheEndPos - PrecacheStartPos);
 #endif
diff --git a/Engine/Source/Runtime/CoreUObject/Private/Serialization/BulkDataStreaming.cpp b/Engine/Source/Runtime/CoreUObject/Private/Serialization/BulkDataStreaming.cpp
--- a/Engine/Source/Runtime/CoreUObject/Private/Serialization/BulkDataStreaming.cpp
+++ b/Engine/Source/Runtime/CoreUObject/Private/Serialization/BulkDataStreaming.cpp
@@ -206,6 +206,7 @@ public:
 	
 	virtual void WaitCompletionImpl(float TimeLimitSeconds) override;
 	virtual void CancelImpl() override;
+	virtual void ReleaseMemoryOwnershipImpl() override;
 	virtual void HandleChunkResult(TIoStatusOr<FIoBuffer>&& Result) override;
 };
 
@@ -241,6 +242,10 @@ void FChunkReadFileRequest::CancelImpl()
 	CancelChunkRequest();
 }
 
+void FChunkReadFileRequest::ReleaseMemoryOwnershipImpl()
+{
+}
+
 void FChunkReadFileRequest::HandleChunkResult(TIoStatusOr<FIoBuffer>&& Result)
 {
 	bCanceled = Result.Status().IsOk() == false;
@@ -279,6 +284,10 @@ private:
 	virtual void CancelImpl() override
 	{
 	}
+
+	virtual void ReleaseMemoryOwnershipImpl() override
+	{
+	}
 };
 
 //////////////////////////////////////////////////////////////////////////////
diff --git a/Engine/Source/Runtime/Engine/Private/AudioStreaming.cpp b/Engine/Source/Runtime/Engine/Private/AudioStreaming.cpp
--- a/Engine/Source/Runtime/Engine/Private/AudioStreaming.cpp
+++ b/Engine/Source/Runtime/Engine/Private/AudioStreaming.cpp
@@ -532,7 +532,6 @@ void FLegacyAudioStreamingManager::ProcessPendingAsyncFileResults()
 
 		ChunkStorage->Data = AudioChunkLoadResult->DataResults;
 
-		DEC_MEMORY_STAT_BY(STAT_AsyncFileMemory, ChunkStorage->DataSize);
 		INC_DWORD_STAT_BY(STAT_AudioMemorySize, ChunkStorage->DataSize);
 		INC_DWORD_STAT_BY(STAT_AudioMemory, ChunkStorage->DataSize);
 
diff --git a/Engine/Source/Runtime/Engine/Private/UnrealEngine.cpp b/Engine/Source/Runtime/Engine/Private/UnrealEngine.cpp
--- a/Engine/Source/Runtime/Engine/Private/UnrealEngine.cpp
+++ b/Engine/Source/Runtime/Engine/Private/UnrealEngine.cpp
@@ -18524,7 +18524,6 @@ static void PakFileTest(const TArray<FString>& Args)
 								FMemory::Memcpy(Memory + CurrentOffset, Mem, Span);
 								bCancel = false; // we should have the memory anyway
 								FMemory::Free(Mem);
-								DEC_MEMORY_STAT_BY(STAT_AsyncFileMemory, Span);
 							}
 						}
 						ReadReq->WaitCompletion();
diff --git a/Engine/Source/Runtime/PakFile/Private/IPlatformFilePak.cpp b/Engine/Source/Runtime/PakFile/Private/IPlatformFilePak.cpp
--- a/Engine/Source/Runtime/PakFile/Private/IPlatformFilePak.cpp
+++ b/Engine/Source/Runtime/PakFile/Private/IPlatformFilePak.cpp
@@ -527,7 +527,11 @@ public:
 		while (!*(volatile bool*)&bCompleteAndCallbackCalled);
 	}
 
-	virtual void CancelImpl()
+	virtual void CancelImpl() override
+	{
+	}
+
+	virtual void ReleaseMemoryOwnershipImpl() override
 	{
 	}
 };
@@ -2715,7 +2719,6 @@ private: // below here we assume CachedFilesScopeLock until we get to the next s
 		if (Block.InRequestRefCount == 0 || bWasCanceled)
 		{
 			check(Block.Size > 0);
-			DEC_MEMORY_STAT_BY(STAT_AsyncFileMemory, Block.Size);
 			FMemory::Free(Memory);
 			UE_LOG(LogPakFile, VeryVerbose, TEXT("FPakReadRequest[%016llX, %016llX) Cancelled"), Block.OffsetAndPakIndex, Block.OffsetAndPakIndex + Block.Size);
 			ClearBlock(Block);
@@ -2726,7 +2729,6 @@ private: // below here we assume CachedFilesScopeLock until we get to the next s
 			check(Block.Memory && Block.Size);
 			BlockMemory += Block.Size;
 			check(BlockMemory > 0);
-			DEC_MEMORY_STAT_BY(STAT_AsyncFileMemory, Block.Size);
 			check(Block.Size > 0);
 			INC_MEMORY_STAT_BY(STAT_PakCacheMem, Block.Size);
 
@@ -3407,6 +3409,11 @@ public:
 		}
 	}
 
+	virtual void ReleaseMemoryOwnershipImpl() override
+	{
+		DEC_MEMORY_STAT_BY(STAT_AsyncFileMemory, BytesToRead);
+	}
+
 	FCachedAsyncBlock& GetBlock()
 	{
 		check(bInternalRequest && BlockPtr);
@@ -3667,6 +3674,11 @@ public:
 		}
 	}
 
+	virtual void ReleaseMemoryOwnershipImpl() override
+	{
+		DEC_MEMORY_STAT_BY(STAT_AsyncFileMemory, BytesToRead);
+	}
+
 	void RequestIsComplete()
 	{
 		// Owner->CriticalSection is locked
@@ -4116,13 +4128,14 @@ public:
 				FMemory::Free(Block.Raw);
 				Block.Raw = nullptr;
 				check(Block.RawSize > 0);
-				DEC_MEMORY_STAT_BY(STAT_AsyncFileMemory, Block.RawSize);
 				Block.RawSize = 0;
 			}
 		}
 		else
 		{
 			check(Block.Raw);
+			// Even though Raw memory has already been loaded, we're treating it as part of the AsyncFileMemory budget until it's processed.
+			INC_MEMORY_STAT_BY(STAT_AsyncFileMemory, Block.RawSize);
 			Block.ProcessedSize = FileEntry.CompressionBlockSize;
 			if (Block.BlockIndex == Blocks.Num() - 1)
 			{
