Adding some Pak initialization and deinitialization logging to track down a horde only ensure.

diff --git a/Engine/Source/Runtime/PakFile/Private/IPlatformFilePak.cpp b/Engine/Source/Runtime/PakFile/Private/IPlatformFilePak.cpp
--- a/Engine/Source/Runtime/PakFile/Private/IPlatformFilePak.cpp
+++ b/Engine/Source/Runtime/PakFile/Private/IPlatformFilePak.cpp
@@ -7387,6 +7387,8 @@ FPakPlatformFile::FPakPlatformFile()
 
 FPakPlatformFile::~FPakPlatformFile()
 {
+	UE_LOG(LogPakFile, Log, TEXT("Destroying PakPlatformFile"));
+
 	FTSTicker::GetCoreTicker().RemoveTicker(RetireReadersHandle);
 
 	FCoreDelegates::GetRegisterEncryptionKeyMulticastDelegate().RemoveAll(this);
@@ -7542,6 +7544,8 @@ static bool PakPlatformFile_IsForceUseIoStore(const TCHAR* CmdLine)
 
 bool FPakPlatformFile::Initialize(IPlatformFile* Inner, const TCHAR* CmdLine)
 {
+	UE_LOG(LogPakFile, Log, TEXT("Initializing PakPlatformFile"));
+
 	LLM_SCOPE(ELLMTag::FileSystem);
 	SCOPED_BOOT_TIMING("FPakPlatformFile::Initialize");
 	// Inner is required.
