PR #10540: Fixed assert message in IPlatformFilePak.cpp

diff --git a/Engine/Source/Runtime/PakFile/Private/IPlatformFilePak.cpp b/Engine/Source/Runtime/PakFile/Private/IPlatformFilePak.cpp
--- a/Engine/Source/Runtime/PakFile/Private/IPlatformFilePak.cpp
+++ b/Engine/Source/Runtime/PakFile/Private/IPlatformFilePak.cpp
@@ -5048,7 +5048,7 @@ public:
 		FName CompressionMethod = PakFile.GetInfo().GetCompressionMethod(PakEntry.CompressionMethodIndex);
 		checkf(FCompression::IsFormatValid(CompressionMethod), 
 			TEXT("Attempting to use compression format %s when loading a file from a .pak, but that compression format is not available.\n")
-			TEXT("If you are running a program (like UnrealPak) you may need to pass the .uproject on the commandline so the plugin can be found.\n"),
+			TEXT("If you are running a program (like UnrealPak) you may need to pass the .uproject on the commandline so the plugin can be found.\n")
 			TEXT("It's also possible that a necessary compression plugin has not been loaded yet, and this file needs to be forced to use zlib compression.\n")
 			TEXT("Unfortunately, the code that can check this does not have the context of the filename that is being read. You will need to look in the callstack in a debugger.\n")
 			TEXT("See ExtensionsToNotUsePluginCompression in [Pak] section of Engine.ini to add more extensions."),
