Speculative fixes for hang detected in ReleaseOldReaders
Acquire list of pak files and then release pak-list lock.
Only release if we can acquire the reader lock without waiting.

diff --git a/Engine/Source/Runtime/PakFile/Private/IPlatformFilePak.cpp b/Engine/Source/Runtime/PakFile/Private/IPlatformFilePak.cpp
--- a/Engine/Source/Runtime/PakFile/Private/IPlatformFilePak.cpp
+++ b/Engine/Source/Runtime/PakFile/Private/IPlatformFilePak.cpp
@@ -7118,7 +7118,7 @@ void FPakFile::ValidateDirectorySearch(const TSet<FString>& FullFoundFiles, cons
 
 bool FPakFile::RecreatePakReaders(IPlatformFile* LowerLevel)
 {
-	FScopeLock ScopedLock(&CriticalSection);
+	FScopeLock ScopedLock(&ReadersCriticalSection);
 
 	if (CurrentlyUsedReaders > 0)
 	{
@@ -7155,7 +7155,7 @@ FSharedPakReader FPakFile::GetSharedReader(IPlatformFile* LowerLevel)
 	LLM_SCOPE_BYTAG(PakSharedReaders);
 	FArchive* PakReader = nullptr;
 	{
-		FScopeLock ScopedLock(&CriticalSection);
+		FScopeLock ScopedLock(&ReadersCriticalSection);
 		if (Readers.Num())
 		{
 			FArchiveAndLastAccessTime Reader = Readers.Pop();
@@ -7179,29 +7179,35 @@ FSharedPakReader FPakFile::GetSharedReader(IPlatformFile* LowerLevel)
 
 void FPakFile::ReturnSharedReader(FArchive* Archive)
 {
-	FScopeLock ScopedLock(&CriticalSection);
+	FScopeLock ScopedLock(&ReadersCriticalSection);
 	--CurrentlyUsedReaders;
 	Readers.Push(FArchiveAndLastAccessTime{ TUniquePtr<FArchive>{Archive }, FPlatformTime::Seconds()});
 }
 
 void FPakFile::ReleaseOldReaders(double MaxAgeSeconds)
 {
-	FScopeLock ScopedLock(&CriticalSection);
-	double SearchTime = FPlatformTime::Seconds() - MaxAgeSeconds;
-	for (int32 i = Readers.Num() - 1; i >= 0; --i)
+	if (ReadersCriticalSection.TryLock())
 	{
-		const FArchiveAndLastAccessTime& Reader = Readers[i];
-		if (Reader.LastAccessTime <= SearchTime)
+		ON_SCOPE_EXIT
 		{
-			// Remove this and all readers older than it (pushed before it)
-			Readers.RemoveAt(0, i + 1);
-			break;
+			ReadersCriticalSection.Unlock();
+		};
+		double SearchTime = FPlatformTime::Seconds() - MaxAgeSeconds;
+		for (int32 i = Readers.Num() - 1; i >= 0; --i)
+		{
+			const FArchiveAndLastAccessTime& Reader = Readers[i];
+			if (Reader.LastAccessTime <= SearchTime)
+			{
+				// Remove this and all readers older than it (pushed before it)
+				Readers.RemoveAt(0, i + 1);
+				break;
+			}
 		}
-	}
 
-	if (Readers.Num() == 0 && CurrentlyUsedReaders == 0)
-	{
-		Decryptor.Reset();
+		if (Readers.Num() == 0 && CurrentlyUsedReaders == 0)
+		{
+			Decryptor.Reset();
+		}
 	}
 }
 
@@ -7708,9 +7714,9 @@ void FPakPlatformFile::ReleaseOldReaders()
 		return;
 	}
 
-	FScopeLock ScopedLock(&PakListCritical);
-
-	for (FPakListEntry& Entry : PakFiles)
+	TArray<FPakListEntry> LocalPaks;
+	GetMountedPaks(LocalPaks);
+	for (FPakListEntry& Entry : LocalPaks)
 	{
 		Entry.PakFile->ReleaseOldReaders(GPakReaderReleaseDelay);
 	}
diff --git a/Engine/Source/Runtime/PakFile/Public/IPlatformFilePak.h b/Engine/Source/Runtime/PakFile/Public/IPlatformFilePak.h
--- a/Engine/Source/Runtime/PakFile/Public/IPlatformFilePak.h
+++ b/Engine/Source/Runtime/PakFile/Public/IPlatformFilePak.h
@@ -783,9 +783,9 @@ private:
 	/** List of readers and when they were last used. */
 	TArray<FArchiveAndLastAccessTime> Readers;
 	/** How many readers have been loaned out and not yet returned. If this is >0 we should not destroy the decryptor. */
-	int32 CurrentlyUsedReaders = 0;
+	std::atomic<int32> CurrentlyUsedReaders = 0;
 	/** Critical section for accessing Readers. */
-	FCriticalSection CriticalSection;
+	FCriticalSection ReadersCriticalSection;
 	/** Pak file info (trailer). */
 	FPakInfo Info;
 	/** Mount point. */
