PakFiles: Add logging of the callstack for errors about directory searches that return incomplete results because of pakfilepruning.
#rnx

diff --git a/Engine/Source/Runtime/PakFile/Private/IPlatformFilePak.cpp b/Engine/Source/Runtime/PakFile/Private/IPlatformFilePak.cpp
--- a/Engine/Source/Runtime/PakFile/Private/IPlatformFilePak.cpp
+++ b/Engine/Source/Runtime/PakFile/Private/IPlatformFilePak.cpp
@@ -7098,6 +7098,31 @@ void FPakFile::ValidateDirectorySearch(const TSet<FString>& FullFoundFiles, cons
 	UE_LOG(LogPakFile, Error, TEXT("FindPrunedFilesAtPath('%s') for PakFile '%s' found a different list in the FullDirectory than in the PrunedDirectory. ")
 		TEXT("Change the calling code or add the files to Engine:[Pak]:WildcardsToKeepInPakStringIndex or Engine:[Pak]:IndexValidationIgnore."),
 		InPath, *PakFilename);
+
+#if !NO_LOGGING && !UE_BUILD_SHIPPING
+	// Logging callstacks is expensive (multiple seconds long). Only do it the first time a path is seen, and only for the first
+	// few paths.
+	static TSet<FString> AlreadyLoggedCallstack;
+	static FCriticalSection AlreadyLoggedCallstackLock;
+	constexpr int32 CallstackLogDirsMax = 10;
+	bool bShouldLogCallstack = false;
+	if (AlreadyLoggedCallstack.Num() < CallstackLogDirsMax) // check to avoid taking critical section if unnecessary
+	{
+		FScopeLock AlreadyLoggedCallstackScopeLock(&AlreadyLoggedCallstackLock);
+		if (AlreadyLoggedCallstack.Num() < CallstackLogDirsMax) // check again since other thread may have modified it
+		{
+			bool bAlreadyLogged;
+			AlreadyLoggedCallstack.Add(FString(InPath), &bAlreadyLogged);
+			bShouldLogCallstack = !bAlreadyLogged;
+		}
+	}
+	if (bShouldLogCallstack)
+	{
+		UE_LOG(LogPakFile, Warning, TEXT("Callstack of FindPrunedFilesAtPath('%s'):"), InPath);
+		FDebug::DumpStackTraceToLog(ELogVerbosity::Warning);
+	}
+#endif
+
 	if (MissingFromPruned.Num() > 0)
 	{
 		for (const FString& Missing : MissingFromPruned)
