PakFiles: Downgrade the Warning: MissingPrunedPakFile from error to warning until we implement a fix for the spurious instances of the error.
#rnx

diff --git a/Engine/Source/Runtime/PakFile/Private/IPlatformFilePak.cpp b/Engine/Source/Runtime/PakFile/Private/IPlatformFilePak.cpp
--- a/Engine/Source/Runtime/PakFile/Private/IPlatformFilePak.cpp
+++ b/Engine/Source/Runtime/PakFile/Private/IPlatformFilePak.cpp
@@ -7092,7 +7092,9 @@ void FPakFile::ValidateDirectorySearch(const TSet<FString>& FullFoundFiles, cons
 	MissingFromPruned.Sort();
 	MissingFromFull.Sort();
 
-	UE_LOG(LogPakFile, Error, TEXT("FindPrunedFilesAtPath('%s') for PakFile '%s' found a different list in the FullDirectory than in the PrunedDirectory. ")
+	// TODO: Restore this as an Error once we modify IPlatformFile::IterateDirectoryRecursively to declare its filefilter so we can ignore the spurious
+	// discovered files that are not part of the fully filtered query
+	UE_LOG(LogPakFile, Warning, TEXT("FindPrunedFilesAtPath('%s') for PakFile '%s' found a different list in the FullDirectory than in the PrunedDirectory. ")
 		TEXT("Change the calling code or add the files to Engine:[Pak]:WildcardsToKeepInPakStringIndex or Engine:[Pak]:IndexValidationIgnore."),
 		InPath, *PakFilename);
 
