PlatformFilePak: Move IPlatformFile directory API functions into cpp file for readability and compile time improvements. Remove functions deprectated on or before 4.27.
#rnx

diff --git a/Engine/Source/Runtime/PakFile/Private/IPlatformFilePak.cpp b/Engine/Source/Runtime/PakFile/Private/IPlatformFilePak.cpp
--- a/Engine/Source/Runtime/PakFile/Private/IPlatformFilePak.cpp
+++ b/Engine/Source/Runtime/PakFile/Private/IPlatformFilePak.cpp
@@ -287,6 +287,484 @@ void FPakPlatformFile::GetFilenamesFromIostoreByBlockIndex(const FString& InCont
 	}
 }
 
+bool FPakPlatformFile::DirectoryExistsInPrunedPakFiles(const TCHAR* Directory)
+{
+	FString StandardPath = Directory;
+	FPaths::MakeStandardFilename(StandardPath);
+
+	TArray<FPakListEntry> Paks;
+	GetMountedPaks(Paks);
+
+	// Check all pak files.
+	for (int32 PakIndex = 0; PakIndex < Paks.Num(); PakIndex++)
+	{
+		if (Paks[PakIndex].PakFile->DirectoryExistsInPruned(*StandardPath))
+		{
+			return true;
+		}
+	}
+	return false;
+}
+
+bool FPakPlatformFile::FindFileInPakFiles(TArray<FPakListEntry>& Paks, const TCHAR* Filename,
+	TRefCountPtr<FPakFile>* OutPakFile, FPakEntry* OutEntry)
+{
+	FString StandardFilename(Filename);
+	FPaths::MakeStandardFilename(StandardFilename);
+
+	int32 DeletedReadOrder = -1;
+
+	for (int32 PakIndex = 0; PakIndex < Paks.Num(); PakIndex++)
+	{
+		int32 PakReadOrder = Paks[PakIndex].ReadOrder;
+		if (DeletedReadOrder != -1 && DeletedReadOrder > PakReadOrder)
+		{
+			// Found a delete record in a higher priority patch level, but now we're at a lower priority set.
+			// Don't search further back or we'll find the original, old file.
+			UE_LOG(LogPakFile, Verbose, TEXT("Delete Record: Accepted a delete record for %s"), Filename);
+			return false;
+		}
+
+		FPakFile::EFindResult FindResult = Paks[PakIndex].PakFile->Find(StandardFilename, OutEntry);
+		if (FindResult == FPakFile::EFindResult::Found)
+		{
+			if (OutPakFile != NULL)
+			{
+				*OutPakFile = Paks[PakIndex].PakFile;
+			}
+			UE_CLOG(DeletedReadOrder != -1, LogPakFile, Verbose,
+				TEXT("Delete Record: Ignored delete record for %s - found it in %s instead (asset was moved between chunks)"),
+				Filename, *Paks[PakIndex].PakFile->GetFilename());
+			return true;
+		}
+		else if (FindResult == FPakFile::EFindResult::FoundDeleted)
+		{
+			DeletedReadOrder = PakReadOrder;
+			UE_LOG(LogPakFile, Verbose, TEXT("Delete Record: Found a delete record for %s in %s"),
+				Filename, *Paks[PakIndex].PakFile->GetFilename());
+		}
+	}
+
+	UE_CLOG(DeletedReadOrder != -1, LogPakFile, Warning,
+		TEXT("Delete Record: No lower priority pak files looking for %s. (maybe not downloaded?)"),
+		Filename);
+	return false;
+}
+
+bool FPakPlatformFile::FindFileInPakFiles(const TCHAR* Filename, TRefCountPtr<FPakFile>* OutPakFile,
+	FPakEntry* OutEntry)
+{
+	TArray<FPakListEntry> Paks;
+	GetMountedPaks(Paks);
+
+	return FindFileInPakFiles(Paks, Filename, OutPakFile, OutEntry);
+}
+
+bool FPakPlatformFile::DirectoryExists(const TCHAR* Directory)
+{
+	// Check pak files first.
+	if (DirectoryExistsInPrunedPakFiles(Directory))
+	{
+		return true;
+	}
+	// Directory does not exist in any of the pak files, continue searching using inner platform file.
+	bool Result = LowerLevel->DirectoryExists(Directory);
+	return Result;
+}
+
+bool FPakPlatformFile::CreateDirectory(const TCHAR* Directory)
+{
+	// Directories can be created only under the normal path
+	return LowerLevel->CreateDirectory(Directory);
+}
+
+bool FPakPlatformFile::DeleteDirectory(const TCHAR* Directory)
+{
+	// Even if the same directory exists outside of pak files it will never
+	// get truly deleted from pak and will still be reported by Iterate functions.
+	// Fail in cases like this.
+	if (DirectoryExistsInPrunedPakFiles(Directory))
+	{
+		return false;
+	}
+	// Directory does not exist in pak files so it's safe to delete.
+	return LowerLevel->DeleteDirectory(Directory);
+}
+
+FFileStatData FPakPlatformFile::GetStatData(const TCHAR* FilenameOrDirectory)
+{
+	// Check pak files first.
+	FPakEntry FileEntry;
+	TRefCountPtr<FPakFile> PakFile;
+	if (FindFileInPakFiles(FilenameOrDirectory, &PakFile, &FileEntry))
+	{
+		return FFileStatData(
+			PakFile->GetTimestamp(),
+			PakFile->GetTimestamp(),
+			PakFile->GetTimestamp(),
+			(FileEntry.CompressionMethodIndex != 0) ? FileEntry.UncompressedSize : FileEntry.Size,
+			false,	// IsDirectory
+			true	// IsReadOnly
+		);
+	}
+
+	// Then check pak directories
+	if (DirectoryExistsInPrunedPakFiles(FilenameOrDirectory))
+	{
+		FDateTime DirectoryTimeStamp = FDateTime::MinValue();
+		return FFileStatData(
+			DirectoryTimeStamp,
+			DirectoryTimeStamp,
+			DirectoryTimeStamp,
+			-1,		// FileSize
+			true,	// IsDirectory
+			true	// IsReadOnly
+		);
+	}
+
+	// Fall back to lower level.
+	FFileStatData FileStatData;
+	if (IsNonPakFilenameAllowed(FilenameOrDirectory))
+	{
+		FileStatData = LowerLevel->GetStatData(FilenameOrDirectory);
+	}
+
+	return FileStatData;
+}
+
+namespace UE::PakFile::Private
+{
+
+/** Helper class to filter out files which have already been visited in one of the pak files. */
+class FPreventDuplicatesVisitorBase
+{
+public:
+	/** Visited files. */
+	TSet<FString>& VisitedFiles;
+	FString NormalizedFilename;
+
+	FPreventDuplicatesVisitorBase(TSet<FString>& InVisitedFiles)
+		: VisitedFiles(InVisitedFiles)
+	{
+	}
+
+	bool CheckDuplicate(const TCHAR* FilenameOrDirectory)
+	{
+		NormalizedFilename.Reset();
+		NormalizedFilename.AppendChars(FilenameOrDirectory, TCString<TCHAR>::Strlen(FilenameOrDirectory));
+		FPaths::MakeStandardFilename(NormalizedFilename);
+		if (VisitedFiles.Contains(NormalizedFilename))
+		{
+			return true;
+		}
+		VisitedFiles.Add(NormalizedFilename);
+		return false;
+	}
+};
+
+class FPreventDuplicatesVisitor : public FPreventDuplicatesVisitorBase, public IPlatformFile::FDirectoryVisitor
+{
+public:
+	/** Wrapped visitor. */
+	FDirectoryVisitor& Visitor;
+
+	/** Constructor. */
+	FPreventDuplicatesVisitor(FDirectoryVisitor& InVisitor, TSet<FString>& InVisitedFiles)
+		: FPreventDuplicatesVisitorBase(InVisitedFiles)
+		, Visitor(InVisitor)
+	{}
+	virtual bool Visit(const TCHAR* FilenameOrDirectory, bool bIsDirectory)
+	{
+		if (CheckDuplicate(FilenameOrDirectory))
+		{
+			// Already visited, continue iterating.
+			return true;
+		}
+		return Visitor.Visit(*NormalizedFilename, bIsDirectory);
+	}
+};
+
+}
+
+bool FPakPlatformFile::IterateDirectory(const TCHAR* Directory, IPlatformFile::FDirectoryVisitor& Visitor)
+{
+	return IterateDirectoryInternal(Directory, Visitor, false /* bRecursive */);
+}
+
+bool FPakPlatformFile::IterateDirectoryInternal(const TCHAR* Directory,
+	IPlatformFile::FDirectoryVisitor& Visitor, bool bRecursive)
+{
+	using namespace UE::PakFile::Private;
+
+	TUniqueFunction<bool(const FString&, const FString&, bool, FPakFile&)> VisitFunction =
+		[&Visitor](const FString& Filename, const FString& NormalizedFilename, bool bIsDir, FPakFile& PakFile)
+	{
+		return Visitor.Visit(*NormalizedFilename, bIsDir);
+	};
+	TSet<FString> FilesVisitedInPak;
+	bool Result = IterateDirectoryInternal(Directory, VisitFunction, bRecursive, FilesVisitedInPak);
+	if (Result && LowerLevel->DirectoryExists(Directory))
+	{
+		// Iterate inner filesystem but don't visit any files that were found in the Paks
+		FPreventDuplicatesVisitor PreventDuplicatesVisitor(Visitor, FilesVisitedInPak);
+		IPlatformFile::FDirectoryVisitor& LowerLevelVisitor(
+			// For performance, skip using PreventDuplicatedVisitor if there were no hits in pak
+			FilesVisitedInPak.Num() ? PreventDuplicatesVisitor : Visitor
+		);
+		if (bRecursive)
+		{
+			Result = LowerLevel->IterateDirectoryRecursively(Directory, LowerLevelVisitor);
+		}
+		else
+		{
+			Result = LowerLevel->IterateDirectory(Directory, LowerLevelVisitor);
+		}
+	}
+	return Result;
+}
+
+bool FPakPlatformFile::IterateDirectoryInternal(const TCHAR* Directory,
+	TUniqueFunction<bool(const FString&, const FString&, bool, FPakFile&)>& VisitFunction,
+	bool bRecursive, TSet<FString>& FilesVisitedInPak)
+{
+	bool Result = true;
+
+	TArray<FPakListEntry> Paks;
+	FString StandardDirectory = Directory;
+	FPaths::MakeStandardFilename(StandardDirectory);
+
+	bool bIsDownloadableDir =
+		(
+			FPaths::HasProjectPersistentDownloadDir() &&
+			StandardDirectory.StartsWith(FPaths::ProjectPersistentDownloadDir())
+			) ||
+		StandardDirectory.StartsWith(FPaths::CloudDir());
+
+	// don't look for in pak files for target-only locations
+	if (!bIsDownloadableDir)
+	{
+		GetMountedPaks(Paks);
+	}
+
+	// Iterate pak files first
+	FString NormalizationBuffer;
+	TSet<FString> FilesVisitedInThisPak;
+	for (int32 PakIndex = 0; PakIndex < Paks.Num(); PakIndex++)
+	{
+		FPakFile& PakFile = *Paks[PakIndex].PakFile;
+
+		const bool bIncludeFiles = true;
+		const bool bIncludeFolders = true;
+
+		FilesVisitedInThisPak.Reset();
+		PakFile.FindPrunedFilesAtPathInternal(*StandardDirectory, FilesVisitedInThisPak,
+			bIncludeFiles, bIncludeFolders, bRecursive);
+		for (TSet<FString>::TConstIterator SetIt(FilesVisitedInThisPak); SetIt && Result; ++SetIt)
+		{
+			const FString& Filename = *SetIt;
+			bool bIsDir = Filename.Len() && Filename[Filename.Len() - 1] == '/';
+			const FString* NormalizedFilename;
+			if (bIsDir)
+			{
+				NormalizationBuffer.Reset(Filename.Len());
+				NormalizationBuffer.AppendChars(*Filename, Filename.Len() - 1); // Chop off the trailing /
+				NormalizedFilename = &NormalizationBuffer;
+			}
+			else
+			{
+				NormalizedFilename = &Filename;
+			}
+			if (!FilesVisitedInPak.Contains(*NormalizedFilename))
+			{
+				FilesVisitedInPak.Add(*NormalizedFilename);
+				Result = VisitFunction(Filename, *NormalizedFilename, bIsDir, PakFile) && Result;
+			}
+		}
+	}
+	return Result;
+}
+
+bool FPakPlatformFile::IterateDirectoryRecursively(const TCHAR* Directory, IPlatformFile::FDirectoryVisitor& Visitor)
+{
+	return IterateDirectoryInternal(Directory, Visitor, true /* bRecursive */);
+}
+
+namespace UE::PakFile::Private
+{
+
+class FPreventDuplicatesStatVisitor : public FPreventDuplicatesVisitorBase, public IPlatformFile::FDirectoryStatVisitor
+{
+public:
+	/** Wrapped visitor. */
+	FDirectoryStatVisitor& Visitor;
+
+	/** Constructor. */
+	FPreventDuplicatesStatVisitor(FDirectoryStatVisitor& InVisitor, TSet<FString>& InVisitedFiles)
+		: FPreventDuplicatesVisitorBase(InVisitedFiles)
+		, Visitor(InVisitor)
+	{}
+	virtual bool Visit(const TCHAR* FilenameOrDirectory, const FFileStatData& StatData)
+	{
+		if (CheckDuplicate(FilenameOrDirectory))
+		{
+			// Already visited, continue iterating.
+			return true;
+		}
+		return Visitor.Visit(*NormalizedFilename, StatData);
+	}
+};
+
+}
+
+bool FPakPlatformFile::IterateDirectoryStat(const TCHAR* Directory, IPlatformFile::FDirectoryStatVisitor& Visitor)
+{
+	return IterateDirectoryStatInternal(Directory, Visitor, false /* bRecursive */);
+}
+
+bool FPakPlatformFile::IterateDirectoryStatInternal(const TCHAR* Directory,
+	IPlatformFile::FDirectoryStatVisitor& Visitor, bool bRecursive)
+{
+	using namespace UE::PakFile::Private;
+
+	TUniqueFunction<bool(const FString&, const FString&, bool, FPakFile&)> VisitFunction =
+		[&Visitor, this](const FString& Filename, const FString& NormalizedFilename, bool bIsDir, FPakFile& PakFile)
+	{
+		int64 FileSize = -1;
+		if (!bIsDir)
+		{
+			FPakEntry FileEntry;
+			if (FindFileInPakFiles(*Filename, nullptr, &FileEntry))
+			{
+				FileSize = (FileEntry.CompressionMethodIndex != 0) ? FileEntry.UncompressedSize : FileEntry.Size;
+			}
+		}
+
+		const FFileStatData StatData(
+			PakFile.GetTimestamp(),
+			PakFile.GetTimestamp(),
+			PakFile.GetTimestamp(),
+			FileSize,
+			bIsDir,
+			true	// IsReadOnly
+		);
+
+		return Visitor.Visit(*NormalizedFilename, StatData);
+	};
+
+	TSet<FString> FilesVisitedInPak;
+	bool Result = IterateDirectoryInternal(Directory, VisitFunction, bRecursive, FilesVisitedInPak);
+	if (Result && LowerLevel->DirectoryExists(Directory))
+	{
+		// Iterate inner filesystem but don't visit any files that were found in the Paks
+		FPreventDuplicatesStatVisitor PreventDuplicatesVisitor(Visitor, FilesVisitedInPak);
+		IPlatformFile::FDirectoryStatVisitor& LowerLevelVisitor(
+			// For performance, skip using PreventDuplicatedVisitor if there were no hits in pak
+			FilesVisitedInPak.Num() ? PreventDuplicatesVisitor : Visitor);
+		if (bRecursive)
+		{
+			Result = LowerLevel->IterateDirectoryStatRecursively(Directory, LowerLevelVisitor);
+		}
+		else
+		{
+			Result = LowerLevel->IterateDirectoryStat(Directory, LowerLevelVisitor);
+		}
+	}
+	return Result;
+}
+
+bool FPakPlatformFile::IterateDirectoryStatRecursively(const TCHAR* Directory,
+	IPlatformFile::FDirectoryStatVisitor& Visitor)
+{
+	return IterateDirectoryStatInternal(Directory, Visitor, true/* bRecursive */);
+}
+
+void FPakPlatformFile::FindFiles(TArray<FString>& FoundFiles, const TCHAR* Directory, const TCHAR* FileExtension)
+{
+	if (LowerLevel->DirectoryExists(Directory))
+	{
+		LowerLevel->FindFiles(FoundFiles, Directory, FileExtension);
+	}
+
+	bool bRecursive = false;
+	FindFilesInternal(FoundFiles, Directory, FileExtension, bRecursive);
+}
+
+void FPakPlatformFile::FindFilesRecursively(TArray<FString>& FoundFiles,
+	const TCHAR* Directory, const TCHAR* FileExtension)
+{
+	if (LowerLevel->DirectoryExists(Directory))
+	{
+		LowerLevel->FindFilesRecursively(FoundFiles, Directory, FileExtension);
+	}
+
+	bool bRecursive = true;
+	FindFilesInternal(FoundFiles, Directory, FileExtension, bRecursive);
+}
+
+void FPakPlatformFile::FindFilesInternal(TArray<FString>& FoundFiles,
+	const TCHAR* Directory, const TCHAR* FileExtension, bool bRecursive)
+{
+	TArray<FPakListEntry> Paks;
+	GetMountedPaks(Paks);
+	if (Paks.Num())
+	{
+		TSet<FString> FilesVisited;
+		FilesVisited.Append(FoundFiles);
+
+		FString StandardDirectory = Directory;
+		FString FileExtensionStr = FileExtension;
+		FPaths::MakeStandardFilename(StandardDirectory);
+		bool bIncludeFiles = true;
+		bool bIncludeFolders = false;
+
+		TArray<FString> FilesInPak;
+		FilesInPak.Reserve(64);
+		for (int32 PakIndex = 0; PakIndex < Paks.Num(); PakIndex++)
+		{
+			FPakFile& PakFile = *Paks[PakIndex].PakFile;
+			PakFile.FindPrunedFilesAtPathInternal(*StandardDirectory, FilesInPak,
+				bIncludeFiles, bIncludeFolders, bRecursive);
+		}
+
+		for (const FString& Filename : FilesInPak)
+		{
+			// filter out files by FileExtension
+			if (FileExtensionStr.Len())
+			{
+				if (!Filename.EndsWith(FileExtensionStr))
+				{
+					continue;
+				}
+			}
+
+			// make sure we don't add duplicates to FoundFiles
+			bool bVisited = false;
+			FilesVisited.Add(Filename, &bVisited);
+			if (!bVisited)
+			{
+				FoundFiles.Add(Filename);
+			}
+		}
+	}
+}
+
+bool FPakPlatformFile::DeleteDirectoryRecursively(const TCHAR* Directory)
+{
+	// Can't delete directories existing in pak files. See DeleteDirectory(..) for more info.
+	if (DirectoryExistsInPrunedPakFiles(Directory))
+	{
+		return false;
+	}
+	// Directory does not exist in pak files so it's safe to delete.
+	return LowerLevel->DeleteDirectoryRecursively(Directory);
+}
+
+bool FPakPlatformFile::CreateDirectoryTree(const TCHAR* Directory)
+{
+	// Directories can only be created only under the normal path
+	return LowerLevel->CreateDirectoryTree(Directory);
+}
+
 void FPakPlatformFile::GetPrunedFilenamesInPakFile(const FString& InPakFilename, TArray<FString>& OutFileList)
 {
 	TArray<FPakListEntry> Paks;
@@ -7033,6 +7511,115 @@ void FPakFile::GetPrunedFilenamesInChunk(const TArray<int32>& InChunkIDs, TArray
 	}
 }
 
+
+/**
+ * Search the given FDirectoryIndex for all files under the given Directory.  Helper for FindFilesAtPath, called separately on the DirectoryIndex or Pruned DirectoryIndex. Does not use
+ * FScopedPakDirectoryIndexAccess internally; caller is responsible for calling from within a lock.
+ * Returned paths are full paths (include the mount point)
+ */
+template <class ContainerType>
+void FPakFile::FindFilesAtPathInIndex(const FDirectoryIndex& TargetIndex, ContainerType& OutFiles, const FString& Directory,
+	bool bIncludeFiles, bool bIncludeDirectories, bool bRecursive) const
+{
+	// Early out if MountPoint is not matching directory
+	if (!Directory.StartsWith(MountPoint))
+	{
+		return;
+	}
+
+	FStringView RelativeSearch(FStringView(Directory).RightChop(MountPoint.Len()));
+
+	TArray<FString> DirectoriesInPak; // List of all unique directories at path
+	for (TMap<FString, FPakDirectory>::TConstIterator It(TargetIndex); It; ++It)
+	{
+		// Check if the file is under the specified path.
+		if (FStringView(It.Key()).StartsWith(RelativeSearch))
+		{
+			FString PakPath = PakPathCombine(MountPoint, It.Key());
+			if (bRecursive == true)
+			{
+				// Add everything
+				if (bIncludeFiles)
+				{
+					for (FPakDirectory::TConstIterator DirectoryIt(It.Value()); DirectoryIt; ++DirectoryIt)
+					{
+						OutFiles.Add(PakPathCombine(PakPath, DirectoryIt.Key()));
+					}
+				}
+				if (bIncludeDirectories)
+				{
+					if (Directory != PakPath)
+					{
+						DirectoriesInPak.Add(MoveTemp(PakPath));
+					}
+				}
+			}
+			else
+			{
+				int32 SubDirIndex = PakPath.Len() > Directory.Len() ? PakPath.Find(TEXT("/"), ESearchCase::CaseSensitive, ESearchDir::FromStart, Directory.Len() + 1) : INDEX_NONE;
+				// Add files in the specified folder only.
+				if (bIncludeFiles && SubDirIndex == INDEX_NONE)
+				{
+					for (FPakDirectory::TConstIterator DirectoryIt(It.Value()); DirectoryIt; ++DirectoryIt)
+					{
+						OutFiles.Add(PakPathCombine(PakPath, DirectoryIt.Key()));
+					}
+				}
+				// Add sub-folders in the specified folder only
+				if (bIncludeDirectories && SubDirIndex >= 0)
+				{
+					DirectoriesInPak.AddUnique(PakPath.Left(SubDirIndex + 1));
+				}
+			}
+		}
+	}
+	OutFiles.Append(MoveTemp(DirectoriesInPak));
+}
+
+template <class ContainerType>
+void FPakFile::FindPrunedFilesAtPathInternal(const TCHAR* InPath, ContainerType& OutFiles,
+	bool bIncludeFiles, bool bIncludeDirectories, bool bRecursive) const
+{
+	// Make sure all directory names end with '/'.
+	FString Directory(InPath);
+	MakeDirectoryFromPath(Directory);
+
+	// Check the specified path is under the mount point of this pak file.
+	// The reverse case (MountPoint StartsWith Directory) is needed to properly handle
+	// pak files that are a subdirectory of the actual directory.
+	if (!Directory.StartsWith(MountPoint) && !MountPoint.StartsWith(Directory))
+	{
+		return;
+	}
+
+	FScopedPakDirectoryIndexAccess ScopeAccess(*this);
+#if ENABLE_PAKFILE_RUNTIME_PRUNING_VALIDATE
+	if (ShouldValidatePrunedDirectory())
+	{
+		TSet<FString> FullFoundFiles, PrunedFoundFiles;
+		FindFilesAtPathInIndex(DirectoryIndex, FullFoundFiles, Directory, bIncludeFiles, bIncludeDirectories, bRecursive);
+		FindFilesAtPathInIndex(PrunedDirectoryIndex, PrunedFoundFiles, Directory, bIncludeFiles, bIncludeDirectories, bRecursive);
+		ValidateDirectorySearch(FullFoundFiles, PrunedFoundFiles, InPath);
+
+		for (const FString& FoundFile : FullFoundFiles)
+		{
+			OutFiles.Add(FoundFile);
+		}
+	}
+	else
+#endif
+	{
+		FindFilesAtPathInIndex(DirectoryIndex, OutFiles, Directory, bIncludeFiles, bIncludeDirectories, bRecursive);
+	}
+}
+
+void FPakFile::FindPrunedFilesAtPath(const TCHAR* InPath, TArray<FString>& OutFiles,
+	bool bIncludeFiles, bool bIncludeDirectories, bool bRecursive) const
+{
+	FindPrunedFilesAtPathInternal(InPath, OutFiles, bIncludeFiles, bIncludeDirectories, bRecursive);
+}
+
+
 #if ENABLE_PAKFILE_RUNTIME_PRUNING_VALIDATE
 void FPakFile::ValidateDirectorySearch(const TSet<FString>& FullFoundFiles, const TSet<FString>& PrunedFoundFiles, const TCHAR* InPath) const
 {
@@ -8248,23 +8835,6 @@ IPakFile* FPakPlatformFile::HandleMountPakDelegate(const FString& PakFilePath, i
 	return nullptr;
 }
 
-bool FPakPlatformFile::HandleOnMountPakDelegate(const FString& PakFilePath, int32 PakOrder, IPlatformFile::FDirectoryVisitor* Visitor)
-{
-	IPakFile* PakFile = HandleMountPakDelegate(PakFilePath, PakOrder);
-	if (PakFile)
-	{
-		if (Visitor)
-		{
-			PakFile->PakVisitPrunedFilenames(*Visitor);
-		}
-		return true;
-	}
-	else
-	{
-		return false;
-	}
-}
-
 bool FPakPlatformFile::HandleUnmountPakDelegate(const FString& PakFilePath)
 {
 	FPlatformMisc::LowLevelOutputDebugStringf(TEXT("Unmounting pak file: %s \n"), *PakFilePath);
@@ -8368,6 +8938,17 @@ PRAGMA_ENABLE_DEPRECATION_WARNINGS
 	return Result;
 }
 
+IFileHandle* FPakPlatformFile::OpenWrite(const TCHAR* Filename, bool bAppend, bool bAllowRead)
+{
+	// No modifications allowed on pak files.
+	if (FindFileInPakFiles(Filename))
+	{
+		return nullptr;
+	}
+	// Use lower level to handle writing.
+	return LowerLevel->OpenWrite(Filename, bAppend, bAllowRead);
+}
+
 const TCHAR* FPakPlatformFile::GetMountStartupPaksWildCard()
 {
 	return *GMountStartupPaksWildCard;
diff --git a/Engine/Source/Runtime/PakFile/Public/IPlatformFilePak.h b/Engine/Source/Runtime/PakFile/Public/IPlatformFilePak.h
--- a/Engine/Source/Runtime/PakFile/Public/IPlatformFilePak.h
+++ b/Engine/Source/Runtime/PakFile/Public/IPlatformFilePak.h
@@ -1041,53 +1041,33 @@ public:
 		return MountPoint;
 	}
 
+	template <class ContainerType>
+	UE_DEPRECATED(5.4, "Use version that takes a TArray<FString>& instead.")
+	void FindPrunedFilesAtPath(ContainerType& OutFiles, const TCHAR* InPath,
+		bool bIncludeFiles = true, bool bIncludeDirectories = false, bool bRecursive = false) const
+	{
+		TArray<FString> LocalFiles;
+		FindPrunedFilesAtPath(InPath, LocalFiles, bIncludeFiles, bIncludeDirectories, bRecursive);
+		for (FString& File : LocalFiles)
+		{
+			OutFiles.Add(File);
+		}
+	}
+
 	/**
 	 * Looks for files or directories within the Pruned DirectoryIndex of the pak file.
 	 * The Pruned DirectoryIndex does not have entries for most Files in the pak; they were removed to save memory.
 	 * A project can specify which FileNames and DirectoryNames can be marked to keep in the DirectoryIndex; see FPakFile::FIndexSettings and FPakFile::PruneDirectoryIndex
 	 * Returned paths are full paths (include the mount point)
 	 *
-	 * @param OutFiles List of files or folder matching search criteria.
 	 * @param InPath Path to look for files or folder at.
+	 * @param OutFiles List of files or folder matching search criteria.
 	 * @param bIncludeFiles If true OutFiles will include matching files.
 	 * @param bIncludeDirectories If true OutFiles will include matching folders.
 	 * @param bRecursive If true, sub-folders will also be checked.
 	 */
-	template <class ContainerType>
-	void FindPrunedFilesAtPath(ContainerType& OutFiles, const TCHAR* InPath, bool bIncludeFiles = true, bool bIncludeDirectories = false, bool bRecursive = false) const
-	{
-		// Make sure all directory names end with '/'.
-		FString Directory(InPath);
-		MakeDirectoryFromPath(Directory);
-
-		// Check the specified path is under the mount point of this pak file.
-		// The reverse case (MountPoint StartsWith Directory) is needed to properly handle
-		// pak files that are a subdirectory of the actual directory.
-		if (!Directory.StartsWith(MountPoint) && !MountPoint.StartsWith(Directory))
-		{
-			return;
-		}
-
-		FScopedPakDirectoryIndexAccess ScopeAccess(*this);
-#if ENABLE_PAKFILE_RUNTIME_PRUNING_VALIDATE
-		if (ShouldValidatePrunedDirectory())
-		{
-			TSet<FString> FullFoundFiles, PrunedFoundFiles;
-			FindFilesAtPathInIndex(DirectoryIndex, FullFoundFiles, Directory, bIncludeFiles, bIncludeDirectories, bRecursive);
-			FindFilesAtPathInIndex(PrunedDirectoryIndex, PrunedFoundFiles, Directory, bIncludeFiles, bIncludeDirectories, bRecursive);
-			ValidateDirectorySearch(FullFoundFiles, PrunedFoundFiles, InPath);
-
-			for (const FString& FoundFile : FullFoundFiles)
-			{
-				OutFiles.Add(FoundFile);
-			}
-		}
-		else
-#endif
-		{
-			FindFilesAtPathInIndex(DirectoryIndex, OutFiles, Directory, bIncludeFiles, bIncludeDirectories, bRecursive);
-		}
-	}
+	PAKFILE_API void FindPrunedFilesAtPath(const TCHAR* InPath, TArray<FString>& OutFiles,
+		bool bIncludeFiles = true, bool bIncludeDirectories = false, bool bRecursive = false) const;
 
 	/**
 	 * Finds a directory in pak file.
@@ -1356,18 +1336,6 @@ public:
 		using FBaseIterator::Filename;
 	};
 
-	class FFileIterator : public FFilenameIterator
-	{
-	public:
-		UE_DEPRECATED(4.26, "FFileIterator is deprecated, use FFilenameIterator instead.  Note that FFilenameIterator will only iterate over the DirectoryIndexKeepFiles entries.")
-		FFileIterator(const FPakFile& InPakFile, bool bInIncludeDeleted = false)
-			: FFilenameIterator(InPakFile, bInIncludeDeleted)
-		{
-		}
-
-		using FFilenameIterator::Filename;
-	};
-
 	/**
 	 * Gets this pak file info.
 	 *
@@ -1691,68 +1659,18 @@ private:
 	/* Manually add a file to a pak file */
 	PAKFILE_API void AddSpecialFile(const FPakEntry& Entry, const FString& Filename);
 
+	template <class ContainerType>
+	void FindPrunedFilesAtPathInternal(const TCHAR* InPath, ContainerType& OutFiles,
+		bool bIncludeFiles = true, bool bIncludeDirectories = false, bool bRecursive = false) const;
+
 	/**
 	 * Search the given FDirectoryIndex for all files under the given Directory.  Helper for FindFilesAtPath, called separately on the DirectoryIndex or Pruned DirectoryIndex. Does not use
 	 * FScopedPakDirectoryIndexAccess internally; caller is responsible for calling from within a lock.
 	 * Returned paths are full paths (include the mount point)
 	 */
 	template <class ContainerType>
-	void FindFilesAtPathInIndex(const FDirectoryIndex& TargetIndex, ContainerType& OutFiles, const FString& Directory, bool bIncludeFiles = true, bool bIncludeDirectories = false, bool bRecursive = false) const
-	{
-		// Early out if MountPoint is not matching directory
-		if (!Directory.StartsWith(MountPoint))
-		{
-			return;
-		}
-
-		FStringView RelativeSearch(FStringView(Directory).RightChop(MountPoint.Len()));
-
-		TArray<FString> DirectoriesInPak; // List of all unique directories at path
-		for (TMap<FString, FPakDirectory>::TConstIterator It(TargetIndex); It; ++It)
-		{
-			// Check if the file is under the specified path.
-			if (FStringView(It.Key()).StartsWith(RelativeSearch))
-			{
-				FString PakPath = PakPathCombine(MountPoint, It.Key());
-				if (bRecursive == true)
-				{
-					// Add everything
-					if (bIncludeFiles)
-					{
-						for (FPakDirectory::TConstIterator DirectoryIt(It.Value()); DirectoryIt; ++DirectoryIt)
-						{
-							OutFiles.Add(PakPathCombine(PakPath, DirectoryIt.Key()));
-						}
-					}
-					if (bIncludeDirectories)
-					{
-						if (Directory != PakPath)
-						{
-							DirectoriesInPak.Add(MoveTemp(PakPath));
-						}
-					}
-				}
-				else
-				{
-					int32 SubDirIndex = PakPath.Len() > Directory.Len() ? PakPath.Find(TEXT("/"), ESearchCase::CaseSensitive, ESearchDir::FromStart, Directory.Len() + 1) : INDEX_NONE;
-					// Add files in the specified folder only.
-					if (bIncludeFiles && SubDirIndex == INDEX_NONE)
-					{
-						for (FPakDirectory::TConstIterator DirectoryIt(It.Value()); DirectoryIt; ++DirectoryIt)
-						{
-							OutFiles.Add(PakPathCombine(PakPath, DirectoryIt.Key()));
-						}
-					}
-					// Add sub-folders in the specified folder only
-					if (bIncludeDirectories && SubDirIndex >= 0)
-					{
-						DirectoriesInPak.AddUnique(PakPath.Left(SubDirIndex + 1));
-					}
-				}
-			}
-		}
-		OutFiles.Append(MoveTemp(DirectoriesInPak));
-	}
+	void FindFilesAtPathInIndex(const FDirectoryIndex& TargetIndex, ContainerType& OutFiles, const FString& Directory,
+		bool bIncludeFiles = true, bool bIncludeDirectories = false, bool bRecursive = false) const;
 
 	/** Converts the path to a RelativePathFromMount and normalizes it to the expected format for Pak Directories.  Returns false if Path is not under the MountDir and hence can not be in this PakFile. */
 	PAKFILE_API bool NormalizeDirectoryQuery(const TCHAR* InPath, FString& OutRelativePathFromMount) const;
@@ -1878,12 +1796,6 @@ class FPakFileHandle : public IFileHandle
 
 public:
 
-	UE_DEPRECATED(4.27, "Use constructor that takes a TRefCountPtr<FPakFile> instead")
-		FPakFileHandle(const FPakFile& InPakFile, const FPakEntry& InPakEntry, TAcquirePakReaderFunction& InAcquirePakReaderFunction)
-		: FPakFileHandle(TRefCountPtr<const FPakFile>(&InPakFile), InPakEntry, InAcquirePakReaderFunction)
-	{
-	}
-
 	/**
 	 * Constructs pak file handle to read from pak.
 	 *
@@ -1899,12 +1811,6 @@ public:
 		INC_DWORD_STAT(STAT_PakFile_NumOpenHandles);
 	}
 
-	UE_DEPRECATED(4.27, "Use constructor that takes a TRefCountPtr<FPakFile> instead")
-		FPakFileHandle(const FPakFile& InPakFile, const FPakEntry& InPakEntry, FArchive* InPakReader)
-		: FPakFileHandle(TRefCountPtr<const FPakFile>(&InPakFile), InPakEntry, InPakReader)
-	{
-	}
-
 	/**
 	 * Constructs pak file handle to read from pak.
 	 *
@@ -2068,36 +1974,13 @@ class FPakPlatformFile : public IPlatformFile
 		Paks.Append(PakFiles);
 	}
 
-	UE_DEPRECATED(4.26, "Use DirectoryExistsInPrunedPakFiles instead")
-	bool DirectoryExistsInPakFiles(const TCHAR* Directory)
-	{
-		return DirectoryExistsInPrunedPakFiles(Directory);
-	}
-
 	/**
 	 * Checks if a directory exists in one of the available pak files.
 	 *
 	 * @param Directory Directory to look for.
 	 * @return true if the directory exists, false otherwise.
 	 */
-	bool DirectoryExistsInPrunedPakFiles(const TCHAR* Directory)
-	{
-		FString StandardPath = Directory;
-		FPaths::MakeStandardFilename(StandardPath);
-
-		TArray<FPakListEntry> Paks;
-		GetMountedPaks(Paks);
-
-		// Check all pak files.
-		for (int32 PakIndex = 0; PakIndex < Paks.Num(); PakIndex++)
-		{
-			if (Paks[PakIndex].PakFile->DirectoryExistsInPruned(*StandardPath))
-			{
-				return true;
-			}
-		}
-		return false;
-	}
+	PAKFILE_API bool DirectoryExistsInPrunedPakFiles(const TCHAR* Directory);
 
 	/**
 	 * Helper function to copy a file from one handle to another usuing the supplied buffer.
@@ -2132,12 +2015,6 @@ class FPakPlatformFile : public IPlatformFile
 	 */
 	PAKFILE_API IPakFile* HandleMountPakDelegate(const FString& PakFilePath, int32 PakOrder);
 
-	/**
-	 * Handler for device delegate to prompt us to load a new pak.
-	 */
-	UE_DEPRECATED(4.26, "Use HandleMountPakDelegate instead")
-	PAKFILE_API bool HandleOnMountPakDelegate(const FString& PakFilePath, int32 PakOrder, IPlatformFile::FDirectoryVisitor* Visitor);
-
 	/**
 	 * Handler for device delegate to prompt us to unload a pak.
 	 */
@@ -2315,25 +2192,6 @@ public:
 	 */
 	PAKFILE_API virtual void MakeUniquePakFilesForTheseFiles(const TArray<TArray<FString>>& InFiles);
 
-
-	/** Overload needed for deprecation; remove this when removing the version with a FPakFile** OutPakFile */
-	static bool FindFileInPakFiles(TArray<FPakListEntry>& Paks, const TCHAR* Filename, nullptr_t OutPakFile, FPakEntry* OutEntry = nullptr)
-	{
-		return FindFileInPakFiles(Paks, Filename, (TRefCountPtr<FPakFile>*) nullptr, OutEntry);
-	}
-
-	UE_DEPRECATED(4.27, "Use version with OutPakFile is a TRefCountPtr<FPakFile> instead")
-	static bool FindFileInPakFiles(TArray<FPakListEntry>& Paks, const TCHAR* Filename, FPakFile** OutPakFile, FPakEntry* OutEntry = nullptr)
-	{
-		TRefCountPtr<FPakFile> PakFile;
-		bool bResult = FindFileInPakFiles(Paks, Filename, &PakFile, OutEntry);
-		if (OutPakFile)
-		{
-			*OutPakFile = PakFile.GetReference();
-		}
-		return bResult;
-	}
-
 	/**
 	 * Finds a file in the specified pak files.
 	 *
@@ -2342,61 +2200,8 @@ public:
 	 * @param OutPakFile Optional pointer to a pak file where the filename was found.
 	 * @return Pointer to pak entry if the file was found, NULL otherwise.
 	 */
-	static bool FindFileInPakFiles(TArray<FPakListEntry>& Paks,const TCHAR* Filename,TRefCountPtr<FPakFile>* OutPakFile,FPakEntry* OutEntry = nullptr)
-	{
-		FString StandardFilename(Filename);
-		FPaths::MakeStandardFilename(StandardFilename);
-
-		int32 DeletedReadOrder = -1;
-
-		for (int32 PakIndex = 0; PakIndex < Paks.Num(); PakIndex++)
-		{
-			int32 PakReadOrder = Paks[PakIndex].ReadOrder;
-			if (DeletedReadOrder != -1 && DeletedReadOrder > PakReadOrder)
-			{
-				//found a delete record in a higher priority patch level, but now we're at a lower priority set - don't search further back or we'll find the original, old file.
-				UE_LOG( LogPakFile, Verbose, TEXT("Delete Record: Accepted a delete record for %s"), Filename );
-				return false;
-			}
-
-			FPakFile::EFindResult FindResult = Paks[PakIndex].PakFile->Find(StandardFilename, OutEntry);
-			if (FindResult == FPakFile::EFindResult::Found )
-			{
-				if (OutPakFile != NULL)
-				{
-					*OutPakFile = Paks[PakIndex].PakFile;
-				}
-				UE_CLOG( DeletedReadOrder != -1, LogPakFile, Verbose, TEXT("Delete Record: Ignored delete record for %s - found it in %s instead (asset was moved between chunks)"), Filename, *Paks[PakIndex].PakFile->GetFilename() );
-				return true;
-			}
-			else if (FindResult == FPakFile::EFindResult::FoundDeleted )
-			{
-				DeletedReadOrder = PakReadOrder;
-				UE_LOG( LogPakFile, Verbose, TEXT("Delete Record: Found a delete record for %s in %s"), Filename, *Paks[PakIndex].PakFile->GetFilename() );
-			}
-		}
-
-		UE_CLOG( DeletedReadOrder != -1, LogPakFile, Warning, TEXT("Delete Record: No lower priority pak files looking for %s. (maybe not downloaded?)"), Filename );
-		return false;
-	}
-
-	/** Overload needed for deprecation; remove this when removing the version with a FPakFile** OutPakFile */
-	bool FindFileInPakFiles(const TCHAR* Filename, nullptr_t OutPakFile, FPakEntry* OutEntry = nullptr)
-	{
-		return FindFileInPakFiles(Filename, (TRefCountPtr<FPakFile>*)nullptr, OutEntry);
-	}
-
-	UE_DEPRECATED(4.27, "Use version with OutPakFile is a TRefCountPtr<FPakFile> instead")
-	bool FindFileInPakFiles(const TCHAR* Filename, FPakFile** OutPakFile, FPakEntry* OutEntry = nullptr)
-	{
-		TRefCountPtr<FPakFile> PakFile;
-		bool bResult = FindFileInPakFiles(Filename, &PakFile, OutEntry);
-		if (OutPakFile)
-		{
-			*OutPakFile = PakFile.GetReference();
-		}
-		return bResult;
-	}
+	PAKFILE_API static bool FindFileInPakFiles(TArray<FPakListEntry>& Paks, const TCHAR* Filename,
+		TRefCountPtr<FPakFile>* OutPakFile, FPakEntry* OutEntry = nullptr);
 
 	/**
 	 * Finds a file in all available pak files.
@@ -2405,13 +2210,8 @@ public:
 	 * @param OutPakFile Optional pointer to a pak file where the filename was found.
 	 * @return Pointer to pak entry if the file was found, NULL otherwise.
 	 */
-	bool FindFileInPakFiles(const TCHAR* Filename, TRefCountPtr<FPakFile>* OutPakFile = nullptr, FPakEntry* OutEntry = nullptr)
-	{
-		TArray<FPakListEntry> Paks;
-		GetMountedPaks(Paks);
-
-		return FindFileInPakFiles(Paks, Filename, OutPakFile, OutEntry);
-	}
+	bool FindFileInPakFiles(const TCHAR* Filename, TRefCountPtr<FPakFile>* OutPakFile = nullptr,
+		FPakEntry* OutEntry = nullptr);
 
 	//~ Begin IPlatformFile Interface
 	virtual bool FileExists(const TCHAR* Filename) override
@@ -2643,389 +2443,27 @@ public:
 	}
 
 	PAKFILE_API virtual IFileHandle* OpenRead(const TCHAR* Filename, bool bAllowWrite = false) override;
-
-	virtual IFileHandle* OpenWrite(const TCHAR* Filename, bool bAppend = false, bool bAllowRead = false) override
-	{
-		// No modifications allowed on pak files.
-		if (FindFileInPakFiles(Filename))
-		{
-			return nullptr;
-		}
-		// Use lower level to handle writing.
-		return LowerLevel->OpenWrite(Filename, bAppend, bAllowRead);
-	}
-
-	virtual bool DirectoryExists(const TCHAR* Directory) override
-	{
-		// Check pak files first.
-		if (DirectoryExistsInPrunedPakFiles(Directory))
-		{
-			return true;
-		}
-		// Directory does not exist in any of the pak files, continue searching using inner platform file.
-		bool Result = LowerLevel->DirectoryExists(Directory); 
-		return Result;
-	}
-
-	virtual bool CreateDirectory(const TCHAR* Directory) override
-	{
-		// Directories can be created only under the normal path
-		return LowerLevel->CreateDirectory(Directory);
-	}
-
-	virtual bool DeleteDirectory(const TCHAR* Directory) override
-	{
-		// Even if the same directory exists outside of pak files it will never
-		// get truly deleted from pak and will still be reported by Iterate functions.
-		// Fail in cases like this.
-		if (DirectoryExistsInPrunedPakFiles(Directory))
-		{
-			return false;
-		}
-		// Directory does not exist in pak files so it's safe to delete.
-		return LowerLevel->DeleteDirectory(Directory);
-	}
-
-	virtual FFileStatData GetStatData(const TCHAR* FilenameOrDirectory) override
-	{
-		// Check pak files first.
-		FPakEntry FileEntry;
-		TRefCountPtr<FPakFile> PakFile;
-		if (FindFileInPakFiles(FilenameOrDirectory, &PakFile, &FileEntry))
-		{
-			return FFileStatData(
-				PakFile->GetTimestamp(),
-				PakFile->GetTimestamp(),
-				PakFile->GetTimestamp(),
-				(FileEntry.CompressionMethodIndex != 0) ? FileEntry.UncompressedSize : FileEntry.Size,
-				false,	// IsDirectory
-				true	// IsReadOnly
-				);
-		}
-
-		// Then check pak directories
-		if (DirectoryExistsInPrunedPakFiles(FilenameOrDirectory))
-		{
-			FDateTime DirectoryTimeStamp = FDateTime::MinValue();
-			return FFileStatData(
-				DirectoryTimeStamp,
-				DirectoryTimeStamp,
-				DirectoryTimeStamp,
-				-1,		// FileSize
-				true,	// IsDirectory
-				true	// IsReadOnly
-				);
-		}
-
-		// Fall back to lower level.
-		FFileStatData FileStatData;
-		if (IsNonPakFilenameAllowed(FilenameOrDirectory))
-		{
-			FileStatData = LowerLevel->GetStatData(FilenameOrDirectory);
-		}
-
-		return FileStatData;
-	}
-
-	/**
-	 * Helper class to filter out files which have already been visited in one of the pak files.
-	 */
-	class FPreventDuplicatesVisitorBase
-	{
-	public:
-		/** Visited files. */
-		TSet<FString>& VisitedFiles;
-		FString NormalizedFilename;
-
-		FPreventDuplicatesVisitorBase(TSet<FString>& InVisitedFiles)
-			: VisitedFiles(InVisitedFiles)
-		{
-		}
-
-		bool CheckDuplicate(const TCHAR* FilenameOrDirectory)
-		{
-			NormalizedFilename.Reset();
-			NormalizedFilename.AppendChars(FilenameOrDirectory, TCString<TCHAR>::Strlen(FilenameOrDirectory));
-			FPaths::MakeStandardFilename(NormalizedFilename);
-			if (VisitedFiles.Contains(NormalizedFilename))
-			{
-				return true;
-			}
-			VisitedFiles.Add(NormalizedFilename);
-			return false;
-		}
-	};
-
-	class FPreventDuplicatesVisitor : public FPreventDuplicatesVisitorBase, public IPlatformFile::FDirectoryVisitor
-	{
-	public:
-		/** Wrapped visitor. */
-		FDirectoryVisitor& Visitor;
-
-		/** Constructor. */
-		FPreventDuplicatesVisitor(FDirectoryVisitor& InVisitor, TSet<FString>& InVisitedFiles)
-			: FPreventDuplicatesVisitorBase(InVisitedFiles)
-			, Visitor(InVisitor)
-		{}
-		virtual bool Visit(const TCHAR* FilenameOrDirectory, bool bIsDirectory)
-		{
-			if (CheckDuplicate(FilenameOrDirectory))
-			{
-				// Already visited, continue iterating.
-				return true;
-			}
-			return Visitor.Visit(*NormalizedFilename, bIsDirectory);
-		}
-	};
-
-	virtual bool IterateDirectory(const TCHAR* Directory, IPlatformFile::FDirectoryVisitor& Visitor) override
-	{
-		return IterateDirectoryInternal(Directory, Visitor, false /* bRecursive */);
-	}
-
-	bool IterateDirectoryInternal(const TCHAR* Directory, IPlatformFile::FDirectoryVisitor& Visitor, bool bRecursive)
-	{
-		TUniqueFunction<bool(const FString&, const FString&, bool, FPakFile&)> VisitFunction = [&Visitor](const FString& Filename, const FString& NormalizedFilename, bool bIsDir, FPakFile& PakFile)
-		{
-			return Visitor.Visit(*NormalizedFilename, bIsDir);
-		};
-		TSet<FString> FilesVisitedInPak;
-		bool Result = IterateDirectoryInternal(Directory, VisitFunction, bRecursive, FilesVisitedInPak);
-		if (Result && LowerLevel->DirectoryExists(Directory))
-		{
-			// Iterate inner filesystem but don't visit any files that were found in the Paks
-			FPreventDuplicatesVisitor PreventDuplicatesVisitor(Visitor, FilesVisitedInPak);
-			IPlatformFile::FDirectoryVisitor& LowerLevelVisitor(FilesVisitedInPak.Num() ? PreventDuplicatesVisitor : Visitor); // For performance, skip using PreventDuplicatedVisitor if there were no hits in pak
-			if (bRecursive)
-			{
-				Result = LowerLevel->IterateDirectoryRecursively(Directory, LowerLevelVisitor);
-			}
-			else
-			{
-				Result = LowerLevel->IterateDirectory(Directory, LowerLevelVisitor);
-			}
-		}
-		return Result;
-	}
-
-	bool IterateDirectoryInternal(const TCHAR* Directory, TUniqueFunction<bool(const FString&, const FString&, bool, FPakFile&)>& VisitFunction, bool bRecursive, TSet<FString>& FilesVisitedInPak)
-	{
-		bool Result = true;
-
-		TArray<FPakListEntry> Paks;
-		FString StandardDirectory = Directory;
-		FPaths::MakeStandardFilename(StandardDirectory);
-		
-		bool bIsDownloadableDir = (FPaths::HasProjectPersistentDownloadDir() && StandardDirectory.StartsWith(FPaths::ProjectPersistentDownloadDir())) || StandardDirectory.StartsWith(FPaths::CloudDir());
-
-		// don't look for in pak files for target-only locations
-		if (!bIsDownloadableDir)
-		{
-			GetMountedPaks(Paks);
-		}
-
-		// Iterate pak files first
-		FString NormalizationBuffer;
-		TSet<FString> FilesVisitedInThisPak;
-		for (int32 PakIndex = 0; PakIndex < Paks.Num(); PakIndex++)
-		{
-			FPakFile& PakFile = *Paks[PakIndex].PakFile;
-			
-			const bool bIncludeFiles = true;
-			const bool bIncludeFolders = true;
-
-			FilesVisitedInThisPak.Reset();
-			PakFile.FindPrunedFilesAtPath(FilesVisitedInThisPak, *StandardDirectory, bIncludeFiles, bIncludeFolders, bRecursive);
-			for (TSet<FString>::TConstIterator SetIt(FilesVisitedInThisPak); SetIt && Result; ++SetIt)
-			{
-				const FString& Filename = *SetIt;
-				bool bIsDir = Filename.Len() && Filename[Filename.Len() - 1] == '/';
-				const FString* NormalizedFilename;
-				if (bIsDir)
-				{
-					NormalizationBuffer.Reset(Filename.Len());
-					NormalizationBuffer.AppendChars(*Filename, Filename.Len()-1); // Chop off the trailing /
-					NormalizedFilename = &NormalizationBuffer;
-				}
-				else
-				{
-					NormalizedFilename = &Filename;
-				}
-				if (!FilesVisitedInPak.Contains(*NormalizedFilename))
-				{
-					FilesVisitedInPak.Add(*NormalizedFilename);
-					Result = VisitFunction(Filename, *NormalizedFilename, bIsDir, PakFile) && Result;
-				}
-			}
-		}
-		return Result;
-	}
-
-	virtual bool IterateDirectoryRecursively(const TCHAR* Directory, IPlatformFile::FDirectoryVisitor& Visitor) override
-	{
-		return IterateDirectoryInternal(Directory, Visitor, true /* bRecursive */);
-	}
-
-	class FPreventDuplicatesStatVisitor : public FPreventDuplicatesVisitorBase, public IPlatformFile::FDirectoryStatVisitor
-	{
-	public:
-		/** Wrapped visitor. */
-		FDirectoryStatVisitor& Visitor;
-
-		/** Constructor. */
-		FPreventDuplicatesStatVisitor(FDirectoryStatVisitor& InVisitor, TSet<FString>& InVisitedFiles)
-			: FPreventDuplicatesVisitorBase(InVisitedFiles)
-			, Visitor(InVisitor)
-		{}
-		virtual bool Visit(const TCHAR* FilenameOrDirectory, const FFileStatData& StatData)
-		{
-			if (CheckDuplicate(FilenameOrDirectory))
-			{
-				// Already visited, continue iterating.
-				return true;
-			}
-			return Visitor.Visit(*NormalizedFilename, StatData);
-		}
-	};
-
-	virtual bool IterateDirectoryStat(const TCHAR* Directory, IPlatformFile::FDirectoryStatVisitor& Visitor) override
-	{
-		return IterateDirectoryStatInternal(Directory, Visitor, false /* bRecursive */);
-	}
-
-	bool IterateDirectoryStatInternal(const TCHAR* Directory, IPlatformFile::FDirectoryStatVisitor& Visitor, bool bRecursive)
-	{
-		TUniqueFunction<bool(const FString&, const FString&, bool, FPakFile&)> VisitFunction = [&Visitor, this](const FString& Filename, const FString& NormalizedFilename, bool bIsDir, FPakFile& PakFile)
-		{
-			int64 FileSize = -1;
-			if (!bIsDir)
-			{
-				FPakEntry FileEntry;
-				if (FindFileInPakFiles(*Filename, nullptr, &FileEntry))
-				{
-					FileSize = (FileEntry.CompressionMethodIndex != 0) ? FileEntry.UncompressedSize : FileEntry.Size;
-				}
-			}
-
-			const FFileStatData StatData(
-				PakFile.GetTimestamp(),
-				PakFile.GetTimestamp(),
-				PakFile.GetTimestamp(),
-				FileSize,
-				bIsDir,
-				true	// IsReadOnly
-			);
-
-			return Visitor.Visit(*NormalizedFilename, StatData);
-		};
-
-		TSet<FString> FilesVisitedInPak;
-		bool Result = IterateDirectoryInternal(Directory, VisitFunction, bRecursive, FilesVisitedInPak);
-		if (Result && LowerLevel->DirectoryExists(Directory))
-		{
-			// Iterate inner filesystem but don't visit any files that were found in the Paks
-			FPreventDuplicatesStatVisitor PreventDuplicatesVisitor(Visitor, FilesVisitedInPak);
-			IPlatformFile::FDirectoryStatVisitor& LowerLevelVisitor(FilesVisitedInPak.Num() ? PreventDuplicatesVisitor : Visitor); // For performance, skip using PreventDuplicatedVisitor if there were no hits in pak
-			if (bRecursive)
-			{
-				Result = LowerLevel->IterateDirectoryStatRecursively(Directory, LowerLevelVisitor);
-			}
-			else
-			{
-				Result = LowerLevel->IterateDirectoryStat(Directory, LowerLevelVisitor);
-			}
-		}
-		return Result;
-	}
-
-	virtual bool IterateDirectoryStatRecursively(const TCHAR* Directory, IPlatformFile::FDirectoryStatVisitor& Visitor) override
-	{
-		return IterateDirectoryStatInternal(Directory, Visitor, true/* bRecursive */);
-	}
-
-	virtual void FindFiles(TArray<FString>& FoundFiles, const TCHAR* Directory, const TCHAR* FileExtension) override
-	{		
-		if (LowerLevel->DirectoryExists(Directory))
-		{
-			LowerLevel->FindFiles(FoundFiles, Directory, FileExtension);
-		}
-
-		bool bRecursive = false;
-		FindFilesInternal(FoundFiles, Directory, FileExtension, bRecursive);
-	}
-	
-	virtual void FindFilesRecursively(TArray<FString>& FoundFiles, const TCHAR* Directory, const TCHAR* FileExtension) override
-	{
-		if (LowerLevel->DirectoryExists(Directory))
-		{
-			LowerLevel->FindFilesRecursively(FoundFiles, Directory, FileExtension);
-		}
-		
-		bool bRecursive = true;
-		FindFilesInternal(FoundFiles, Directory, FileExtension, bRecursive);
-	}
-
-	void FindFilesInternal(TArray<FString>& FoundFiles, const TCHAR* Directory, const TCHAR* FileExtension, bool bRecursive)
-	{
-		TArray<FPakListEntry> Paks;
-		GetMountedPaks(Paks);
-		if (Paks.Num())
-		{
-			TSet<FString> FilesVisited;
-			FilesVisited.Append(FoundFiles);
-			
-			FString StandardDirectory = Directory;
-			FString FileExtensionStr = FileExtension;
-			FPaths::MakeStandardFilename(StandardDirectory);
-			bool bIncludeFiles = true;
-			bool bIncludeFolders = false;
-
-			TArray<FString> FilesInPak;
-			FilesInPak.Reserve(64);
-			for (int32 PakIndex = 0; PakIndex < Paks.Num(); PakIndex++)
-			{
-				FPakFile& PakFile = *Paks[PakIndex].PakFile;
-				PakFile.FindPrunedFilesAtPath(FilesInPak, *StandardDirectory, bIncludeFiles, bIncludeFolders, bRecursive);
-			}
-			
-			for (const FString& Filename : FilesInPak)
-			{
-				// filter out files by FileExtension
-				if (FileExtensionStr.Len())
-				{
-					if (!Filename.EndsWith(FileExtensionStr))
-					{
-						continue;
-					}
-				}
-								
-				// make sure we don't add duplicates to FoundFiles
-				bool bVisited = false;
-				FilesVisited.Add(Filename, &bVisited);
-				if (!bVisited)
-				{
-					FoundFiles.Add(Filename);
-				}
-			}
-		}
-	}
-
-	virtual bool DeleteDirectoryRecursively(const TCHAR* Directory) override
-	{
-		// Can't delete directories existing in pak files. See DeleteDirectory(..) for more info.
-		if (DirectoryExistsInPrunedPakFiles(Directory))
-		{
-			return false;
-		}
-		// Directory does not exist in pak files so it's safe to delete.
-		return LowerLevel->DeleteDirectoryRecursively(Directory);
-	}
-
-	virtual bool CreateDirectoryTree(const TCHAR* Directory) override
-	{
-		// Directories can only be created only under the normal path
-		return LowerLevel->CreateDirectoryTree(Directory);
-	}
+	PAKFILE_API virtual IFileHandle* OpenWrite(const TCHAR* Filename, bool bAppend = false,
+		bool bAllowRead = false) override;
+
+	PAKFILE_API virtual bool DirectoryExists(const TCHAR* Directory) override;
+	PAKFILE_API virtual bool CreateDirectory(const TCHAR* Directory) override;
+	PAKFILE_API virtual bool DeleteDirectory(const TCHAR* Directory) override;
+	PAKFILE_API virtual FFileStatData GetStatData(const TCHAR* FilenameOrDirectory) override;
+	PAKFILE_API virtual bool IterateDirectory(const TCHAR* Directory,
+		IPlatformFile::FDirectoryVisitor& Visitor) override;
+	PAKFILE_API virtual bool IterateDirectoryRecursively(const TCHAR* Directory,
+		IPlatformFile::FDirectoryVisitor& Visitor) override;
+	PAKFILE_API virtual bool IterateDirectoryStat(const TCHAR* Directory,
+		IPlatformFile::FDirectoryStatVisitor& Visitor) override;
+	PAKFILE_API virtual bool IterateDirectoryStatRecursively(const TCHAR* Directory,
+		IPlatformFile::FDirectoryStatVisitor& Visitor) override;
+	PAKFILE_API virtual void FindFiles(TArray<FString>& FoundFiles,
+		const TCHAR* Directory, const TCHAR* FileExtension) override;
+	PAKFILE_API virtual void FindFilesRecursively(TArray<FString>& FoundFiles,
+		const TCHAR* Directory, const TCHAR* FileExtension) override;
+	PAKFILE_API virtual bool DeleteDirectoryRecursively(const TCHAR* Directory) override;
+	PAKFILE_API virtual bool CreateDirectoryTree(const TCHAR* Directory) override;
 
 	PAKFILE_API virtual bool CopyFile(const TCHAR* To, const TCHAR* From, EPlatformFileRead ReadFlags = EPlatformFileRead::None, EPlatformFileWrite WriteFlags = EPlatformFileWrite::None) override;
 
@@ -3150,6 +2588,16 @@ public:
 	// Internal cache of pak signature files
 	PAKFILE_API static TMap<FName, TSharedPtr<const struct FPakSignatureFile, ESPMode::ThreadSafe>> PakSignatureFileCache;
 	PAKFILE_API static FCriticalSection PakSignatureFileCacheLock;
+
+private:
+	bool IterateDirectoryInternal(const TCHAR* Directory, IPlatformFile::FDirectoryVisitor& Visitor, bool bRecursive);
+	bool IterateDirectoryInternal(const TCHAR* Directory,
+		TUniqueFunction<bool(const FString&, const FString&, bool, FPakFile&)>& VisitFunction,
+		bool bRecursive, TSet<FString>& FilesVisitedInPak);
+	bool IterateDirectoryStatInternal(const TCHAR* Directory,
+		IPlatformFile::FDirectoryStatVisitor& Visitor, bool bRecursive);
+	void FindFilesInternal(TArray<FString>& FoundFiles,
+		const TCHAR* Directory, const TCHAR* FileExtension, bool bRecursive);
 };
 
 /**
