PakFiles and FileManager: Add IPlatformFile::FDirectoryVisitor::ShouldVisitLeafPath for directory visitors that are ignoring some files, so that FPakFile::ValidateDirectorySearch knows that invalid-to-request files were not requested despite being present in the searched directories.
#rnx

diff --git a/Engine/Plugins/Developer/Concert/ConcertSync/ConcertSyncClient/Source/ConcertSyncClient/Private/ConcertSandboxPlatformFile.cpp b/Engine/Plugins/Developer/Concert/ConcertSync/ConcertSyncClient/Source/ConcertSyncClient/Private/ConcertSandboxPlatformFile.cpp
--- a/Engine/Plugins/Developer/Concert/ConcertSync/ConcertSyncClient/Source/ConcertSyncClient/Private/ConcertSandboxPlatformFile.cpp
+++ b/Engine/Plugins/Developer/Concert/ConcertSync/ConcertSyncClient/Source/ConcertSyncClient/Private/ConcertSandboxPlatformFile.cpp
@@ -593,7 +593,7 @@ bool FConcertSandboxPlatformFile::IterateDirectory(const TCHAR* Directory, FDire
 		const TArray<FDirectoryItem> DirectoryItems = GetDirectoryContents(ResolvedPath, Directory);
 		for (const FDirectoryItem& DirectoryItem : DirectoryItems)
 		{
-			if (!Visitor.Visit(*DirectoryItem.Path, DirectoryItem.StatData.bIsDirectory))
+			if (!Visitor.CallShouldVisitAndVisit(*DirectoryItem.Path, DirectoryItem.StatData.bIsDirectory))
 			{
 				return false;
 			}
@@ -619,7 +619,7 @@ bool FConcertSandboxPlatformFile::IterateDirectoryStat(const TCHAR* Directory, F
 		const TArray<FDirectoryItem> DirectoryItems = GetDirectoryContents(ResolvedPath, Directory);
 		for (const FDirectoryItem& DirectoryItem : DirectoryItems)
 		{
-			if (!Visitor.Visit(*DirectoryItem.Path, DirectoryItem.StatData))
+			if (!Visitor.CallShouldVisitAndVisit(*DirectoryItem.Path, DirectoryItem.StatData))
 			{
 				return false;
 			}
diff --git a/Engine/Source/Runtime/Core/Private/Android/AndroidPlatformFile.cpp b/Engine/Source/Runtime/Core/Private/Android/AndroidPlatformFile.cpp
--- a/Engine/Source/Runtime/Core/Private/Android/AndroidPlatformFile.cpp
+++ b/Engine/Source/Runtime/Core/Private/Android/AndroidPlatformFile.cpp
@@ -1963,12 +1963,12 @@ public:
 		auto InternalVisitor = [&](const FString& InLocalPath, struct dirent* InEntry) -> bool
 		{
 			const FString DirPath = DirectoryStr / UTF8_TO_TCHAR(InEntry->d_name);
-			return Visitor.Visit(*DirPath, InEntry->d_type == DT_DIR);
+			return Visitor.CallShouldVisitAndVisit(*DirPath, InEntry->d_type == DT_DIR);
 		};
 		
 		auto InternalResourceVisitor = [&](const FString& InResourceName, bool IsDirectory) -> bool
 		{
-			return Visitor.Visit(*InResourceName, IsDirectory);
+			return Visitor.CallShouldVisitAndVisit(*InResourceName, IsDirectory);
 		};
 		
 		auto InternalAssetVisitor = [&](const char* InAssetPath) -> bool
@@ -1981,7 +1981,7 @@ public:
 				AAssetDir_close(subdir);
 			}
 
-			return Visitor.Visit(UTF8_TO_TCHAR(InAssetPath), isDirectory);
+			return Visitor.CallShouldVisitAndVisit(UTF8_TO_TCHAR(InAssetPath), isDirectory);
 		};
 
 		return IterateDirectoryCommon(Directory, InternalVisitor, InternalResourceVisitor, InternalAssetVisitor, AllowLocal, AllowAsset);
@@ -2003,7 +2003,7 @@ public:
 			struct stat FileInfo;
 			if (stat(TCHAR_TO_UTF8(*(InLocalPath / UTF8_TO_TCHAR(InEntry->d_name))), &FileInfo) != -1)
 			{
-				return Visitor.Visit(*DirPath, AndroidStatToUEFileData(FileInfo));
+				return Visitor.CallShouldVisitAndVisit(*DirPath, AndroidStatToUEFileData(FileInfo));
 			}
 
 			return true;
@@ -2011,7 +2011,7 @@ public:
 		
 		auto InternalResourceVisitor = [&](const FString& InResourceName, bool IsDir) -> bool
 		{
-			return Visitor.Visit(
+			return Visitor.CallShouldVisitAndVisit(
 				*InResourceName, 
 				FFileStatData(
 					FDateTime::MinValue(),						// CreationTime
@@ -2042,7 +2042,7 @@ public:
 				AAssetDir_close(subdir);
 			}
 
-			return Visitor.Visit(
+			return Visitor.CallShouldVisitAndVisit(
 				UTF8_TO_TCHAR(InAssetPath), 
 				FFileStatData(
 					FDateTime::MinValue(),	// CreationTime
diff --git a/Engine/Source/Runtime/Core/Private/Apple/ApplePlatformFile.cpp b/Engine/Source/Runtime/Core/Private/Apple/ApplePlatformFile.cpp
--- a/Engine/Source/Runtime/Core/Private/Apple/ApplePlatformFile.cpp
+++ b/Engine/Source/Runtime/Core/Private/Apple/ApplePlatformFile.cpp
@@ -743,7 +743,7 @@ bool FApplePlatformFile::IterateDirectory(const TCHAR* Directory, FDirectoryVisi
 			}
 		}
 
-		return Visitor.Visit(*(DirectoryStr / NormalizedFilename), bIsDirectory);
+		return Visitor.CallShouldVisitAndVisit(*(DirectoryStr / NormalizedFilename), bIsDirectory);
 	});
 }
 
@@ -762,7 +762,7 @@ bool FApplePlatformFile::IterateDirectoryStat(const TCHAR* Directory, FDirectory
 		struct stat StatInfo;
 		if (stat(TCHAR_TO_UTF8(*(NormalizedDirectoryStr / NormalizedFilename)), &StatInfo) == 0)
 		{
-			return Visitor.Visit(*(DirectoryStr / NormalizedFilename), MacStatToUEFileData(StatInfo));
+			return Visitor.CallShouldVisitAndVisit(*(DirectoryStr / NormalizedFilename), MacStatToUEFileData(StatInfo));
 		}
 
 		return true;
diff --git a/Engine/Source/Runtime/Core/Private/GenericPlatform/GenericPlatformFile.cpp b/Engine/Source/Runtime/Core/Private/GenericPlatform/GenericPlatformFile.cpp
--- a/Engine/Source/Runtime/Core/Private/GenericPlatform/GenericPlatformFile.cpp
+++ b/Engine/Source/Runtime/Core/Private/GenericPlatform/GenericPlatformFile.cpp
@@ -4,6 +4,7 @@
 #include "Algo/Accumulate.h"
 #include "HAL/FileManager.h"
 #include "Misc/Paths.h"
+#include "Misc/PathViews.h"
 #include "HAL/PlatformMisc.h"
 #include "HAL/ThreadSafeCounter.h"
 #include "Stats/Stats.h"
@@ -538,6 +539,27 @@ FDateTime IPlatformFile::GetTimeStampLocal(const TCHAR* Filename)
 	return FileTimeStamp;
 }
 
+bool IPlatformFile::FDirectoryVisitor::CallShouldVisitAndVisit(const TCHAR* FilenameOrDirectory, bool bIsDirectory)
+{
+	FStringView LeafPathname = FPathViews::GetCleanFilename(FilenameOrDirectory);
+	if (!ShouldVisitLeafPathname(LeafPathname))
+	{
+		return true; // Continue iterating
+	}
+	return Visit(FilenameOrDirectory, bIsDirectory);
+}
+
+bool IPlatformFile::FDirectoryStatVisitor::CallShouldVisitAndVisit(const TCHAR* FilenameOrDirectory,
+	const FFileStatData& StatData)
+{
+	FStringView LeafPathname = FPathViews::GetCleanFilename(FilenameOrDirectory);
+	if (!ShouldVisitLeafPathname(LeafPathname))
+	{
+		return true; // Continue iterating
+	}
+	return Visit(FilenameOrDirectory, StatData);
+}
+
 class FDirectoryVisitorFuncWrapper : public IPlatformFile::FDirectoryVisitor
 {
 public:
@@ -593,7 +615,7 @@ bool IPlatformFile::IterateDirectoryRecursively(const TCHAR* Directory, FDirecto
 		}
 		virtual bool Visit(const TCHAR* FilenameOrDirectory, bool bIsDirectory) override
 		{
-			bool bResult = Visitor.Visit(FilenameOrDirectory, bIsDirectory);
+			bool bResult = Visitor.CallShouldVisitAndVisit(FilenameOrDirectory, bIsDirectory);
 			if (bResult && bIsDirectory)
 			{
 				Directories.Emplace(FilenameOrDirectory);
@@ -649,7 +671,7 @@ bool IPlatformFile::IterateDirectoryStatRecursively(const TCHAR* Directory, FDir
 		}
 		virtual bool Visit(const TCHAR* FilenameOrDirectory, const FFileStatData& StatData) override
 		{
-			bool bResult = Visitor.Visit(FilenameOrDirectory, StatData);
+			bool bResult = Visitor.CallShouldVisitAndVisit(FilenameOrDirectory, StatData);
 			if (bResult && StatData.bIsDirectory)
 			{
 				bResult = PlatformFile.IterateDirectoryStat(FilenameOrDirectory, *this);
diff --git a/Engine/Source/Runtime/Core/Private/HAL/FileManagerGeneric.cpp b/Engine/Source/Runtime/Core/Private/HAL/FileManagerGeneric.cpp
--- a/Engine/Source/Runtime/Core/Private/HAL/FileManagerGeneric.cpp
+++ b/Engine/Source/Runtime/Core/Private/HAL/FileManagerGeneric.cpp
@@ -427,12 +427,19 @@ namespace FileManagerGenericImpl
 			, bStoreFullPath(bInStoreFullPath)
 		{
 		}
+
+		virtual bool ShouldVisitLeafPathname(FStringView LeafFilename) override
+		{
+			return FString(LeafFilename).MatchesWildcard(WildCard);
+		}
+
 		virtual bool Visit(const TCHAR* FilenameOrDirectory, bool bIsDirectory)
 		{
 			if ((bIsDirectory && bDirectories) || (!bIsDirectory && bFiles))
 			{
 				FString Filename = FPaths::GetCleanFilename(FilenameOrDirectory);
-				if (Filename.MatchesWildcard(WildCard))
+				if (ensureMsgf(ShouldVisitLeafPathname(Filename),
+					TEXT("PlatformFile.IterateDirectory needs to call ShouldVisitLeafFilename before calling Visit.")))
 				{
 					FString FullPath = bStoreFullPath ? FString(FilenameOrDirectory) : MoveTemp(Filename);
 					FWriteScopeLock ScopeLock(ResultLock);
diff --git a/Engine/Source/Runtime/Core/Private/IOS/IOSPlatformFile.cpp b/Engine/Source/Runtime/Core/Private/IOS/IOSPlatformFile.cpp
--- a/Engine/Source/Runtime/Core/Private/IOS/IOSPlatformFile.cpp
+++ b/Engine/Source/Runtime/Core/Private/IOS/IOSPlatformFile.cpp
@@ -940,7 +940,7 @@ bool FIOSPlatformFile::IterateDirectory(const TCHAR* Directory, FDirectoryVisito
 		const FString NormalizedFilename = UTF8_TO_TCHAR(([[[NSString stringWithUTF8String:InEntry->d_name] precomposedStringWithCanonicalMapping] cStringUsingEncoding:NSUTF8StringEncoding]));
 		const FString FullPath = DirectoryStr / NormalizedFilename;
 
-		return Visitor.Visit(*FullPath, InEntry->d_type == DT_DIR);
+		return Visitor.CallShouldVisitAndVisit(*FullPath, InEntry->d_type == DT_DIR);
 	});
 }
 
@@ -972,7 +972,7 @@ bool FIOSPlatformFile::IterateDirectoryStat(const TCHAR* Directory, FDirectorySt
 			}
 		}
 
-		return Visitor.Visit(*FullPath, IOSStatToUEFileData(FileInfo));
+		return Visitor.CallShouldVisitAndVisit(*FullPath, IOSStatToUEFileData(FileInfo));
 	});
 }
 
diff --git a/Engine/Source/Runtime/Core/Private/Unix/UnixPlatformFile.cpp b/Engine/Source/Runtime/Core/Private/Unix/UnixPlatformFile.cpp
--- a/Engine/Source/Runtime/Core/Private/Unix/UnixPlatformFile.cpp
+++ b/Engine/Source/Runtime/Core/Private/Unix/UnixPlatformFile.cpp
@@ -1300,7 +1300,7 @@ bool FUnixPlatformFile::IterateDirectory(const TCHAR* Directory, FDirectoryVisit
 			}
 		}
 
-		return Visitor.Visit(*(DirectoryStr / UnicodeEntryName), bIsDirectory);
+		return Visitor.CallShouldVisitAndVisit(*(DirectoryStr / UnicodeEntryName), bIsDirectory);
 	});
 }
 
@@ -1317,7 +1317,7 @@ bool FUnixPlatformFile::IterateDirectoryStat(const TCHAR* Directory, FDirectoryS
 		const FString AbsoluteUnicodeName = NormalizedDirectoryStr / UnicodeEntryName;	
 		if (stat(TCHAR_TO_UTF8(*AbsoluteUnicodeName), &FileInfo) != -1)
 		{
-			return Visitor.Visit(*(DirectoryStr / UnicodeEntryName), UnixStatToUEFileData(FileInfo));
+			return Visitor.CallShouldVisitAndVisit(*(DirectoryStr / UnicodeEntryName), UnixStatToUEFileData(FileInfo));
 		}
 
 		return true;
diff --git a/Engine/Source/Runtime/Core/Private/Windows/WindowsPlatformFile.cpp b/Engine/Source/Runtime/Core/Private/Windows/WindowsPlatformFile.cpp
--- a/Engine/Source/Runtime/Core/Private/Windows/WindowsPlatformFile.cpp
+++ b/Engine/Source/Runtime/Core/Private/Windows/WindowsPlatformFile.cpp
@@ -1492,7 +1492,7 @@ public:
 			DirTemp.RemoveSuffix(DirTemp.Len() - DirLen);
 			FPathViews::Append(DirTemp, InData.cFileName);
 			const bool bIsDirectory = !!(InData.dwFileAttributes & FILE_ATTRIBUTE_DIRECTORY);
-			return Visitor.Visit(*DirTemp, bIsDirectory);
+			return Visitor.CallShouldVisitAndVisit(*DirTemp, bIsDirectory);
 		});
 	}
 	virtual bool IterateDirectoryStat(const TCHAR* Directory, FDirectoryStatVisitor& Visitor) override
@@ -1511,7 +1511,7 @@ public:
 				FileSize = static_cast<int64>(li.QuadPart);
 			}
 
-			return Visitor.Visit(
+			return Visitor.CallShouldVisitAndVisit(
 				*(DirectoryStr / InData.cFileName), 
 				FFileStatData(
 					WindowsFileTimeToUEDateTime(InData.ftCreationTime),
diff --git a/Engine/Source/Runtime/Core/Public/GenericPlatform/GenericPlatformFile.h b/Engine/Source/Runtime/Core/Public/GenericPlatform/GenericPlatformFile.h
--- a/Engine/Source/Runtime/Core/Public/GenericPlatform/GenericPlatformFile.h
+++ b/Engine/Source/Runtime/Core/Public/GenericPlatform/GenericPlatformFile.h
@@ -394,6 +394,16 @@ public:
 
 		virtual ~FDirectoryVisitor() { }
 
+		/**
+		 * Called with the LeafPathname (FullPath == Path/LeafPathname, LeafPathname == BaseName.Extension) before
+		 * calling Visit. If it returns true, Visit will be called on the path, otherwise Visit will be skipped, and
+		 * the return value of Visit is treated as true (continue iterating). Called both for directories and files.
+		 */
+		virtual bool ShouldVisitLeafPathname(FStringView LeafPathname)
+		{
+			return true;
+		}
+
 		/** 
 		 * Callback for a single file or a directory in a directory iteration.
 		 * @param FilenameOrDirectory		If bIsDirectory is true, this is a directory (with no trailing path delimiter), otherwise it is a file name.
@@ -402,6 +412,12 @@ public:
 		**/
 		virtual bool Visit(const TCHAR* FilenameOrDirectory, bool bIsDirectory) = 0;
 
+		/**
+		 * Helper function for receivers of FDirectoryVisitor. Enforces the contract for whether Visit should be
+		 * called after calling ShouldVisitLeafPathName.
+		 */
+		CORE_API bool CallShouldVisitAndVisit(const TCHAR* FilenameOrDirectory, bool bIsDirectory);
+
 		/** True if the Visit function can be called from multiple threads at once. **/
 		FORCEINLINE bool IsThreadSafe() const
 		{
@@ -419,13 +435,30 @@ public:
 	{
 	public:
 		virtual ~FDirectoryStatVisitor() { }
-		/** 
+
+		/**
+		 * Called with the LeafPathname (FullPath == Path/LeafPathname, LeafPathname == BaseName.Extension) before
+		 * calling Visit. If it returns true, Visit will be called on the path, otherwise Visit will be skipped, and
+		 * the return value of Visit is treated as true (continue iterating). Called both for directories and files.
+		 */
+		virtual bool ShouldVisitLeafPathname(FStringView LeafPathname)
+		{
+			return true;
+		}
+
+		/**
 		 * Callback for a single file or a directory in a directory iteration.
 		 * @param FilenameOrDirectory		If bIsDirectory is true, this is a directory (with no trailing path delimiter), otherwise it is a file name.
 		 * @param StatData					The stat data for the file or directory.
 		 * @return							true if the iteration should continue.
 		**/
 		virtual bool Visit(const TCHAR* FilenameOrDirectory, const FFileStatData& StatData) = 0;
+
+		/**
+		 * Helper function for receivers of FDirectoryStatVisitor. Enforces the contract for whether Visit should be
+		 *  called after calling ShouldVisitLeafPathName.
+		 */
+		CORE_API bool CallShouldVisitAndVisit(const TCHAR* FilenameOrDirectory, const FFileStatData& StatData);
 	};
 
 	/** File and directory visitor function that takes all the stat data */
diff --git a/Engine/Source/Runtime/Core/Public/HAL/IPlatformFileLogWrapper.h b/Engine/Source/Runtime/Core/Public/HAL/IPlatformFileLogWrapper.h
--- a/Engine/Source/Runtime/Core/Public/HAL/IPlatformFileLogWrapper.h
+++ b/Engine/Source/Runtime/Core/Public/HAL/IPlatformFileLogWrapper.h
@@ -365,7 +365,7 @@ public:
 		{
 			FILE_LOG(LogPlatformFile, Verbose, TEXT("Visit %s %d"), FilenameOrDirectory, int32(bIsDirectory));
 			double StartTime = FPlatformTime::Seconds();
-			bool Result = Visitor.Visit(FilenameOrDirectory, bIsDirectory);
+			bool Result = Visitor.CallShouldVisitAndVisit(FilenameOrDirectory, bIsDirectory);
 			double ThisTime = (FPlatformTime::Seconds() - StartTime) / 1000.0;
 			FILE_LOG(LogPlatformFile, Verbose, TEXT("Visit return %d [%fms]"), int32(Result), ThisTime);
 			return Result;
@@ -409,7 +409,7 @@ public:
 		{
 			FILE_LOG(LogPlatformFile, Verbose, TEXT("Visit %s %d"), FilenameOrDirectory, int32(StatData.bIsDirectory));
 			double StartTime = FPlatformTime::Seconds();
-			bool Result = Visitor.Visit(FilenameOrDirectory, StatData);
+			bool Result = Visitor.CallShouldVisitAndVisit(FilenameOrDirectory, StatData);
 			double ThisTime = (FPlatformTime::Seconds() - StartTime) / 1000.0;
 			FILE_LOG(LogPlatformFile, Verbose, TEXT("Visit return %d [%fms]"), int32(Result), ThisTime);
 			return Result;
diff --git a/Engine/Source/Runtime/NetworkFile/Private/NetworkPlatformFile.cpp b/Engine/Source/Runtime/NetworkFile/Private/NetworkPlatformFile.cpp
--- a/Engine/Source/Runtime/NetworkFile/Private/NetworkPlatformFile.cpp
+++ b/Engine/Source/Runtime/NetworkFile/Private/NetworkPlatformFile.cpp
@@ -555,7 +555,7 @@ bool FNetworkPlatformFile::IterateDirectory(const TCHAR* InDirectory, IPlatformF
 				bool bIsDirectory = It.Value() == 0;
 			
 				// visit (stripping off the path if needed)
-				RetVal = Visitor.Visit(bHadNoPath ? *FPaths::GetCleanFilename(It.Key()) : *It.Key(), bIsDirectory);
+				RetVal = Visitor.CallShouldVisitAndVisit(bHadNoPath ? *FPaths::GetCleanFilename(It.Key()) : *It.Key(), bIsDirectory);
 			}
 		}
 	}
@@ -592,7 +592,7 @@ bool FNetworkPlatformFile::IterateDirectoryRecursively(const TCHAR* InDirectory,
 				bool bIsDirectory = It.Value() == 0;
 
 				// visit!
-				RetVal = Visitor.Visit(*It.Key(), bIsDirectory);
+				RetVal = Visitor.CallShouldVisitAndVisit(*It.Key(), bIsDirectory);
 			}
 		}
 	}
@@ -641,7 +641,7 @@ bool FNetworkPlatformFile::IterateDirectoryStat(const TCHAR* InDirectory, IPlatf
 					);
 
 				// visit (stripping off the path if needed)
-				RetVal = Visitor.Visit(bHadNoPath ? *FPaths::GetCleanFilename(It.Key()) : *It.Key(), StatData);
+				RetVal = Visitor.CallShouldVisitAndVisit(bHadNoPath ? *FPaths::GetCleanFilename(It.Key()) : *It.Key(), StatData);
 			}
 		}
 	}
@@ -688,7 +688,7 @@ bool FNetworkPlatformFile::IterateDirectoryStatRecursively(const TCHAR* InDirect
 					);
 
 				// visit!
-				RetVal = Visitor.Visit(*It.Key(), StatData);
+				RetVal = Visitor.CallShouldVisitAndVisit(*It.Key(), StatData);
 			}
 		}
 	}
diff --git a/Engine/Source/Runtime/PakFile/Private/IPlatformFilePak.cpp b/Engine/Source/Runtime/PakFile/Private/IPlatformFilePak.cpp
--- a/Engine/Source/Runtime/PakFile/Private/IPlatformFilePak.cpp
+++ b/Engine/Source/Runtime/PakFile/Private/IPlatformFilePak.cpp
@@ -37,12 +37,14 @@
 #include "Async/MappedFileHandle.h"
 #include "IoDispatcherFileBackend.h"
 #include "Misc/PackageName.h"
+#include "Misc/PathViews.h"
 
 #include "ProfilingDebugging/LoadTimeTracker.h"
 #include "IO/IoContainerHeader.h"
 #include "FilePackageStore.h"
 #include "Compression/OodleDataCompression.h"
 #include "IO/IoStore.h"
+#include "String/RemoveFrom.h"
 
 DEFINE_LOG_CATEGORY(LogPakFile);
 
@@ -480,10 +482,48 @@ public:
 			// Already visited, continue iterating.
 			return true;
 		}
-		return Visitor.Visit(*NormalizedFilename, bIsDirectory);
+		return Visitor.CallShouldVisitAndVisit(*NormalizedFilename, bIsDirectory);
 	}
 };
 
+/**
+ * A file/directory visitor for files in PakFiles, used to share code for FDirectoryVisitor and FDirectoryStatVisitor
+ * when iterating over files in pakfiles.
+ */
+class FPakFileDirectoryVisitorBase
+{
+public:
+	FPakFileDirectoryVisitorBase()
+	{
+	}
+	virtual ~FPakFileDirectoryVisitorBase() { }
+
+	virtual bool ShouldVisitLeafPathname(FStringView LeafNormalizedPathname) = 0;
+	virtual bool Visit(const FString& Filename, const FString& NormalizedFilename, bool bIsDir, FPakFile& PakFile) = 0;
+	// No need for CallShouldVisitAndVisit because we call ShouldVisitLeafPathname separately in all cases
+};
+
+/** FPakFileDirectoryVisitorBase for a FDirectoryVisitor. */
+class FPakFileDirectoryVisitor : public FPakFileDirectoryVisitorBase
+{
+public:
+	FPakFileDirectoryVisitor(IPlatformFile::FDirectoryVisitor& InInner)
+		: Inner(InInner)
+	{
+	}
+	virtual bool ShouldVisitLeafPathname(FStringView LeafNormalizedPathname) override
+	{
+		return Inner.ShouldVisitLeafPathname(LeafNormalizedPathname);
+	}
+	virtual bool Visit(const FString& Filename, const FString& NormalizedFilename,
+		bool bIsDir, FPakFile& PakFile) override
+	{
+		return Inner.Visit(*NormalizedFilename, bIsDir);
+	}
+
+	IPlatformFile::FDirectoryVisitor& Inner;
+};
+
 }
 
 bool FPakPlatformFile::IterateDirectory(const TCHAR* Directory, IPlatformFile::FDirectoryVisitor& Visitor)
@@ -496,13 +536,9 @@ bool FPakPlatformFile::IterateDirectoryInternal(const TCHAR* Directory,
 {
 	using namespace UE::PakFile::Private;
 
-	TUniqueFunction<bool(const FString&, const FString&, bool, FPakFile&)> VisitFunction =
-		[&Visitor](const FString& Filename, const FString& NormalizedFilename, bool bIsDir, FPakFile& PakFile)
-	{
-		return Visitor.Visit(*NormalizedFilename, bIsDir);
-	};
+	FPakFileDirectoryVisitor PakVisitor(Visitor);
 	TSet<FString> FilesVisitedInPak;
-	bool Result = IterateDirectoryInternal(Directory, VisitFunction, bRecursive, FilesVisitedInPak);
+	bool Result = IterateDirectoryInPakFiles(Directory, PakVisitor, bRecursive, FilesVisitedInPak);
 	if (Result && LowerLevel->DirectoryExists(Directory))
 	{
 		// Iterate inner filesystem but don't visit any files that were found in the Paks
@@ -523,9 +559,8 @@ bool FPakPlatformFile::IterateDirectoryInternal(const TCHAR* Directory,
 	return Result;
 }
 
-bool FPakPlatformFile::IterateDirectoryInternal(const TCHAR* Directory,
-	TUniqueFunction<bool(const FString&, const FString&, bool, FPakFile&)>& VisitFunction,
-	bool bRecursive, TSet<FString>& FilesVisitedInPak)
+bool FPakPlatformFile::IterateDirectoryInPakFiles(const TCHAR* Directory,
+	UE::PakFile::Private::FPakFileDirectoryVisitorBase& Visitor, bool bRecursive, TSet<FString>& FilesVisitedInPak)
 {
 	bool Result = true;
 
@@ -549,6 +584,11 @@ bool FPakPlatformFile::IterateDirectoryInternal(const TCHAR* Directory,
 	// Iterate pak files first
 	FString NormalizationBuffer;
 	TSet<FString> FilesVisitedInThisPak;
+	auto ShouldVisit = [&Visitor](FStringView UnnormalizedPath)
+	{
+		FStringView NormalizedPath = UE::String::RemoveFromEnd(UnnormalizedPath, TEXTVIEW("/"));
+		return Visitor.ShouldVisitLeafPathname(FPathViews::GetCleanFilename(NormalizedPath));
+	};
 	for (int32 PakIndex = 0; PakIndex < Paks.Num(); PakIndex++)
 	{
 		FPakFile& PakFile = *Paks[PakIndex].PakFile;
@@ -557,7 +597,7 @@ bool FPakPlatformFile::IterateDirectoryInternal(const TCHAR* Directory,
 		const bool bIncludeFolders = true;
 
 		FilesVisitedInThisPak.Reset();
-		PakFile.FindPrunedFilesAtPathInternal(*StandardDirectory, FilesVisitedInThisPak,
+		PakFile.FindPrunedFilesAtPathInternal(*StandardDirectory, ShouldVisit, FilesVisitedInThisPak,
 			bIncludeFiles, bIncludeFolders, bRecursive);
 		for (TSet<FString>::TConstIterator SetIt(FilesVisitedInThisPak); SetIt && Result; ++SetIt)
 		{
@@ -577,7 +617,7 @@ bool FPakPlatformFile::IterateDirectoryInternal(const TCHAR* Directory,
 			if (!FilesVisitedInPak.Contains(*NormalizedFilename))
 			{
 				FilesVisitedInPak.Add(*NormalizedFilename);
-				Result = VisitFunction(Filename, *NormalizedFilename, bIsDir, PakFile) && Result;
+				Result = Visitor.Visit(Filename, *NormalizedFilename, bIsDir, PakFile) && Result;
 			}
 		}
 	}
@@ -610,30 +650,31 @@ public:
 			// Already visited, continue iterating.
 			return true;
 		}
-		return Visitor.Visit(*NormalizedFilename, StatData);
+		return Visitor.CallShouldVisitAndVisit(*NormalizedFilename, StatData);
 	}
 };
 
-}
-
-bool FPakPlatformFile::IterateDirectoryStat(const TCHAR* Directory, IPlatformFile::FDirectoryStatVisitor& Visitor)
-{
-	return IterateDirectoryStatInternal(Directory, Visitor, false /* bRecursive */);
-}
-
-bool FPakPlatformFile::IterateDirectoryStatInternal(const TCHAR* Directory,
-	IPlatformFile::FDirectoryStatVisitor& Visitor, bool bRecursive)
+/** FPakFileDirectoryVisitorBase for a FDirectoryStatVisitor. */
+class FPakFileDirectoryStatVisitor : public FPakFileDirectoryVisitorBase
 {
-	using namespace UE::PakFile::Private;
-
-	TUniqueFunction<bool(const FString&, const FString&, bool, FPakFile&)> VisitFunction =
-		[&Visitor, this](const FString& Filename, const FString& NormalizedFilename, bool bIsDir, FPakFile& PakFile)
+public:
+	FPakFileDirectoryStatVisitor(FPakPlatformFile& InPlatformFile, IPlatformFile::FDirectoryStatVisitor& InInner)
+		: PlatformFile(InPlatformFile)
+		, Inner(InInner)
+	{
+	}
+	virtual bool ShouldVisitLeafPathname(FStringView LeafNormalizedPathname) override
+	{
+		return Inner.ShouldVisitLeafPathname(LeafNormalizedPathname);
+	}
+	virtual bool Visit(const FString& Filename, const FString& NormalizedFilename,
+		bool bIsDir, FPakFile& PakFile) override
 	{
 		int64 FileSize = -1;
 		if (!bIsDir)
 		{
 			FPakEntry FileEntry;
-			if (FindFileInPakFiles(*Filename, nullptr, &FileEntry))
+			if (PlatformFile.FindFileInPakFiles(*Filename, nullptr, &FileEntry))
 			{
 				FileSize = (FileEntry.CompressionMethodIndex != 0) ? FileEntry.UncompressedSize : FileEntry.Size;
 			}
@@ -648,11 +689,28 @@ bool FPakPlatformFile::IterateDirectoryStatInternal(const TCHAR* Directory,
 			true	// IsReadOnly
 		);
 
-		return Visitor.Visit(*NormalizedFilename, StatData);
-	};
+		return Inner.Visit(*NormalizedFilename, StatData);
+	}
+
+	FPakPlatformFile& PlatformFile;
+	IPlatformFile::FDirectoryStatVisitor& Inner;
+};
+
+}
+
+bool FPakPlatformFile::IterateDirectoryStat(const TCHAR* Directory, IPlatformFile::FDirectoryStatVisitor& Visitor)
+{
+	return IterateDirectoryStatInternal(Directory, Visitor, false /* bRecursive */);
+}
+
+bool FPakPlatformFile::IterateDirectoryStatInternal(const TCHAR* Directory,
+	IPlatformFile::FDirectoryStatVisitor& Visitor, bool bRecursive)
+{
+	using namespace UE::PakFile::Private;
 
+	FPakFileDirectoryStatVisitor PakVisitor(*this, Visitor);
 	TSet<FString> FilesVisitedInPak;
-	bool Result = IterateDirectoryInternal(Directory, VisitFunction, bRecursive, FilesVisitedInPak);
+	bool Result = IterateDirectoryInPakFiles(Directory, PakVisitor, bRecursive, FilesVisitedInPak);
 	if (Result && LowerLevel->DirectoryExists(Directory))
 	{
 		// Iterate inner filesystem but don't visit any files that were found in the Paks
@@ -712,31 +770,28 @@ void FPakPlatformFile::FindFilesInternal(TArray<FString>& FoundFiles,
 		FilesVisited.Append(FoundFiles);
 
 		FString StandardDirectory = Directory;
-		FString FileExtensionStr = FileExtension;
+		FStringView FileExtensionStr = FileExtension;
 		FPaths::MakeStandardFilename(StandardDirectory);
 		bool bIncludeFiles = true;
 		bool bIncludeFolders = false;
 
+		auto ShouldVisit = [FileExtensionStr](FStringView Filename)
+		{
+			// filter out files by FileExtension
+			return FileExtensionStr.Len() == 0 || Filename.EndsWith(FileExtensionStr, ESearchCase::IgnoreCase);
+		};
+
 		TArray<FString> FilesInPak;
 		FilesInPak.Reserve(64);
 		for (int32 PakIndex = 0; PakIndex < Paks.Num(); PakIndex++)
 		{
 			FPakFile& PakFile = *Paks[PakIndex].PakFile;
-			PakFile.FindPrunedFilesAtPathInternal(*StandardDirectory, FilesInPak,
+			PakFile.FindPrunedFilesAtPathInternal(*StandardDirectory, ShouldVisit, FilesInPak,
 				bIncludeFiles, bIncludeFolders, bRecursive);
 		}
 
 		for (const FString& Filename : FilesInPak)
 		{
-			// filter out files by FileExtension
-			if (FileExtensionStr.Len())
-			{
-				if (!Filename.EndsWith(FileExtensionStr))
-				{
-					continue;
-				}
-			}
-
 			// make sure we don't add duplicates to FoundFiles
 			bool bVisited = false;
 			FilesVisited.Add(Filename, &bVisited);
@@ -7517,9 +7572,10 @@ void FPakFile::GetPrunedFilenamesInChunk(const TArray<int32>& InChunkIDs, TArray
  * FScopedPakDirectoryIndexAccess internally; caller is responsible for calling from within a lock.
  * Returned paths are full paths (include the mount point)
  */
-template <class ContainerType>
-void FPakFile::FindFilesAtPathInIndex(const FDirectoryIndex& TargetIndex, ContainerType& OutFiles, const FString& Directory,
-	bool bIncludeFiles, bool bIncludeDirectories, bool bRecursive) const
+template <typename ShouldVisitFunc, class ContainerType>
+void FPakFile::FindFilesAtPathInIndex(const FDirectoryIndex& TargetIndex, ContainerType& OutFiles,
+	const FString& Directory, const ShouldVisitFunc& ShouldVisit, bool bIncludeFiles, bool bIncludeDirectories,
+	bool bRecursive) const
 {
 	// Early out if MountPoint is not matching directory
 	if (!Directory.StartsWith(MountPoint))
@@ -7543,14 +7599,21 @@ void FPakFile::FindFilesAtPathInIndex(const FDirectoryIndex& TargetIndex, Contai
 				{
 					for (FPakDirectory::TConstIterator DirectoryIt(It.Value()); DirectoryIt; ++DirectoryIt)
 					{
-						OutFiles.Add(PakPathCombine(PakPath, DirectoryIt.Key()));
+						const FString& FilePathUnderDirectory = DirectoryIt.Key();
+						if (ShouldVisit(FilePathUnderDirectory))
+						{
+							OutFiles.Add(PakPathCombine(PakPath, FilePathUnderDirectory));
+						}
 					}
 				}
 				if (bIncludeDirectories)
 				{
 					if (Directory != PakPath)
 					{
-						DirectoriesInPak.Add(MoveTemp(PakPath));
+						if (ShouldVisit(PakPath))
+						{
+							DirectoriesInPak.Add(MoveTemp(PakPath));
+						}
 					}
 				}
 			}
@@ -7562,13 +7625,21 @@ void FPakFile::FindFilesAtPathInIndex(const FDirectoryIndex& TargetIndex, Contai
 				{
 					for (FPakDirectory::TConstIterator DirectoryIt(It.Value()); DirectoryIt; ++DirectoryIt)
 					{
-						OutFiles.Add(PakPathCombine(PakPath, DirectoryIt.Key()));
+						const FString& FilePathUnderDirectory = DirectoryIt.Key();
+						if (ShouldVisit(FilePathUnderDirectory))
+						{
+							OutFiles.Add(PakPathCombine(PakPath, FilePathUnderDirectory));
+						}
 					}
 				}
 				// Add sub-folders in the specified folder only
 				if (bIncludeDirectories && SubDirIndex >= 0)
 				{
-					DirectoriesInPak.AddUnique(PakPath.Left(SubDirIndex + 1));
+					FString SubDirPath = PakPath.Left(SubDirIndex + 1);
+					if (ShouldVisit(SubDirPath))
+					{
+						DirectoriesInPak.AddUnique(MoveTemp(SubDirPath));
+					}
 				}
 			}
 		}
@@ -7576,8 +7647,8 @@ void FPakFile::FindFilesAtPathInIndex(const FDirectoryIndex& TargetIndex, Contai
 	OutFiles.Append(MoveTemp(DirectoriesInPak));
 }
 
-template <class ContainerType>
-void FPakFile::FindPrunedFilesAtPathInternal(const TCHAR* InPath, ContainerType& OutFiles,
+template <typename ShouldVisitFunc, class ContainerType>
+void FPakFile::FindPrunedFilesAtPathInternal(const TCHAR* InPath, const ShouldVisitFunc& ShouldVisit, ContainerType& OutFiles,
 	bool bIncludeFiles, bool bIncludeDirectories, bool bRecursive) const
 {
 	// Make sure all directory names end with '/'.
@@ -7597,8 +7668,10 @@ void FPakFile::FindPrunedFilesAtPathInternal(const TCHAR* InPath, ContainerType&
 	if (ShouldValidatePrunedDirectory())
 	{
 		TSet<FString> FullFoundFiles, PrunedFoundFiles;
-		FindFilesAtPathInIndex(DirectoryIndex, FullFoundFiles, Directory, bIncludeFiles, bIncludeDirectories, bRecursive);
-		FindFilesAtPathInIndex(PrunedDirectoryIndex, PrunedFoundFiles, Directory, bIncludeFiles, bIncludeDirectories, bRecursive);
+		FindFilesAtPathInIndex(DirectoryIndex, FullFoundFiles, Directory, ShouldVisit,
+			bIncludeFiles, bIncludeDirectories, bRecursive);
+		FindFilesAtPathInIndex(PrunedDirectoryIndex, PrunedFoundFiles, Directory, ShouldVisit,
+			bIncludeFiles, bIncludeDirectories, bRecursive);
 		ValidateDirectorySearch(FullFoundFiles, PrunedFoundFiles, InPath);
 
 		for (const FString& FoundFile : FullFoundFiles)
@@ -7609,14 +7682,16 @@ void FPakFile::FindPrunedFilesAtPathInternal(const TCHAR* InPath, ContainerType&
 	else
 #endif
 	{
-		FindFilesAtPathInIndex(DirectoryIndex, OutFiles, Directory, bIncludeFiles, bIncludeDirectories, bRecursive);
+		FindFilesAtPathInIndex(DirectoryIndex, OutFiles, Directory, ShouldVisit,
+			bIncludeFiles, bIncludeDirectories, bRecursive);
 	}
 }
 
 void FPakFile::FindPrunedFilesAtPath(const TCHAR* InPath, TArray<FString>& OutFiles,
 	bool bIncludeFiles, bool bIncludeDirectories, bool bRecursive) const
 {
-	FindPrunedFilesAtPathInternal(InPath, OutFiles, bIncludeFiles, bIncludeDirectories, bRecursive);
+	auto ShouldVisit = [](FStringView Path) { return true; };
+	FindPrunedFilesAtPathInternal(InPath, ShouldVisit, OutFiles, bIncludeFiles, bIncludeDirectories, bRecursive);
 }
 
 
diff --git a/Engine/Source/Runtime/PakFile/Public/IPlatformFilePak.h b/Engine/Source/Runtime/PakFile/Public/IPlatformFilePak.h
--- a/Engine/Source/Runtime/PakFile/Public/IPlatformFilePak.h
+++ b/Engine/Source/Runtime/PakFile/Public/IPlatformFilePak.h
@@ -55,6 +55,7 @@ class FFilePackageStoreBackend;
 class FOutputDevice;
 class IAsyncReadFileHandle;
 class IMappedFileHandle;
+namespace UE::PakFile::Private { class FPakFileDirectoryVisitorBase; }
 struct FIoContainerHeader;
 
 PAKFILE_API DECLARE_LOG_CATEGORY_EXTERN(LogPakFile, Log, All);
@@ -888,7 +889,7 @@ public:
 	{
 		for (FFilenameIterator It(*this); It; ++It)
 		{
-			Visitor.Visit(*It.Filename(), false);
+			Visitor.CallShouldVisitAndVisit(*It.Filename(), false);
 		}
 	}
 
@@ -1659,8 +1660,8 @@ private:
 	/* Manually add a file to a pak file */
 	PAKFILE_API void AddSpecialFile(const FPakEntry& Entry, const FString& Filename);
 
-	template <class ContainerType>
-	void FindPrunedFilesAtPathInternal(const TCHAR* InPath, ContainerType& OutFiles,
+	template <typename ShouldVisitFunc, class ContainerType>
+	void FindPrunedFilesAtPathInternal(const TCHAR* InPath, const ShouldVisitFunc& ShouldVisit, ContainerType& OutFiles,
 		bool bIncludeFiles = true, bool bIncludeDirectories = false, bool bRecursive = false) const;
 
 	/**
@@ -1668,9 +1669,10 @@ private:
 	 * FScopedPakDirectoryIndexAccess internally; caller is responsible for calling from within a lock.
 	 * Returned paths are full paths (include the mount point)
 	 */
-	template <class ContainerType>
+	template <typename ShouldVisitFunc, class ContainerType>
 	void FindFilesAtPathInIndex(const FDirectoryIndex& TargetIndex, ContainerType& OutFiles, const FString& Directory,
-		bool bIncludeFiles = true, bool bIncludeDirectories = false, bool bRecursive = false) const;
+		const ShouldVisitFunc& ShouldVisit, bool bIncludeFiles = true, bool bIncludeDirectories = false,
+		bool bRecursive = false) const;
 
 	/** Converts the path to a RelativePathFromMount and normalizes it to the expected format for Pak Directories.  Returns false if Path is not under the MountDir and hence can not be in this PakFile. */
 	PAKFILE_API bool NormalizeDirectoryQuery(const TCHAR* InPath, FString& OutRelativePathFromMount) const;
@@ -2591,9 +2593,9 @@ public:
 
 private:
 	bool IterateDirectoryInternal(const TCHAR* Directory, IPlatformFile::FDirectoryVisitor& Visitor, bool bRecursive);
-	bool IterateDirectoryInternal(const TCHAR* Directory,
-		TUniqueFunction<bool(const FString&, const FString&, bool, FPakFile&)>& VisitFunction,
-		bool bRecursive, TSet<FString>& FilesVisitedInPak);
+	bool IterateDirectoryInPakFiles(const TCHAR* Directory,
+		UE::PakFile::Private::FPakFileDirectoryVisitorBase& Visitor, bool bRecursive,
+		TSet<FString>& FilesVisitedInPak);
 	bool IterateDirectoryStatInternal(const TCHAR* Directory,
 		IPlatformFile::FDirectoryStatVisitor& Visitor, bool bRecursive);
 	void FindFilesInternal(TArray<FString>& FoundFiles,
diff --git a/Engine/Source/Runtime/SandboxFile/Private/IPlatformFileSandboxWrapper.cpp b/Engine/Source/Runtime/SandboxFile/Private/IPlatformFileSandboxWrapper.cpp
--- a/Engine/Source/Runtime/SandboxFile/Private/IPlatformFileSandboxWrapper.cpp
+++ b/Engine/Source/Runtime/SandboxFile/Private/IPlatformFileSandboxWrapper.cpp
@@ -744,7 +744,7 @@ public:
 			if (NormalizedFilename == SandboxFile.InjectedSourceDirectoryParent)
 			{
 				// "fake" the injected directory so everything falls into place
-				Visitor.Visit(*SandboxFile.InjectedSourceDirectory, true);
+				Visitor.CallShouldVisitAndVisit(*SandboxFile.InjectedSourceDirectory, true);
 
 				// for recursive visiting, we need to recurse into the forced directory
 				if (bIsRecursive)
@@ -791,7 +791,7 @@ public:
 
 		if (CanVisit)
 		{
-			bool Result = Visitor.Visit(*LocalFilename, bIsDirectory);
+			bool Result = Visitor.CallShouldVisitAndVisit(*LocalFilename, bIsDirectory);
 			return Result;
 		}
 		else
@@ -861,7 +861,7 @@ public:
 			{
 				// "fake" the injected directory so everything falls into place
 				FFileStatData InjectedStat = SandboxFile.LowerLevel->GetStatData(*SandboxFile.InjectedSourceDirectory);
-				Visitor.Visit(*SandboxFile.InjectedSourceDirectory, InjectedStat);
+				Visitor.CallShouldVisitAndVisit(*SandboxFile.InjectedSourceDirectory, InjectedStat);
 
 				// for recursive visiting, we need to recurse into the forced directory
 				if (bIsRecursive)
@@ -906,7 +906,7 @@ public:
 
 		if (CanVisit)
 		{
-			bool Result = Visitor.Visit(*LocalFilename, StatData);
+			bool Result = Visitor.CallShouldVisitAndVisit(*LocalFilename, StatData);
 			return Result;
 		}
 		else
diff --git a/Engine/Source/Runtime/StorageServerClient/Private/StorageServerPlatformFile.cpp b/Engine/Source/Runtime/StorageServerClient/Private/StorageServerPlatformFile.cpp
--- a/Engine/Source/Runtime/StorageServerClient/Private/StorageServerPlatformFile.cpp
+++ b/Engine/Source/Runtime/StorageServerClient/Private/StorageServerPlatformFile.cpp
@@ -551,7 +551,7 @@ bool FStorageServerPlatformFile::IterateDirectory(const TCHAR* Directory, IPlatf
 			bool bConverted = MakeLocalPath(FilenameOrDirectory, LocalPath);
 			check(bConverted);
 			const bool bDirectory = !FileChunkId.IsValid();
-			return Visitor.Visit(*LocalPath, bDirectory);
+			return Visitor.CallShouldVisitAndVisit(*LocalPath, bDirectory);
 		});
 	}
 	else
@@ -588,7 +588,7 @@ bool FStorageServerPlatformFile::IterateDirectoryStat(const TCHAR* Directory, FD
 					true,
 					true);
 			}
-			return Visitor.Visit(*LocalPath, FileStatData);
+			return Visitor.CallShouldVisitAndVisit(*LocalPath, FileStatData);
 		});
 	}
 	else
diff --git a/Engine/Source/Runtime/StreamingFile/Private/StreamingNetworkPlatformFile.cpp b/Engine/Source/Runtime/StreamingFile/Private/StreamingNetworkPlatformFile.cpp
--- a/Engine/Source/Runtime/StreamingFile/Private/StreamingNetworkPlatformFile.cpp
+++ b/Engine/Source/Runtime/StreamingFile/Private/StreamingNetworkPlatformFile.cpp
@@ -603,7 +603,7 @@ bool FStreamingNetworkPlatformFile::IterateDirectory(const TCHAR* InDirectory, I
 				bool bIsDirectory = It.Value() == 0;
 			
 				// visit (stripping off the path if needed)
-				RetVal = Visitor.Visit(bHadNoPath ? *FPaths::GetCleanFilename(It.Key()) : *It.Key(), bIsDirectory);
+				RetVal = Visitor.CallShouldVisitAndVisit(bHadNoPath ? *FPaths::GetCleanFilename(It.Key()) : *It.Key(), bIsDirectory);
 			}
 		}
 	}
@@ -633,7 +633,7 @@ bool FStreamingNetworkPlatformFile::IterateDirectoryRecursively(const TCHAR* InD
 				bool bIsDirectory = It.Value() == 0;
 
 				// visit!
-				RetVal = Visitor.Visit(*It.Key(), bIsDirectory);
+				RetVal = Visitor.CallShouldVisitAndVisit(*It.Key(), bIsDirectory);
 			}
 		}
 	}
