Add a cvar s.IoDispatcherForceSynchronousScatter to allow forcing IODispatcher scatter jobs to always run on the IODispatcher thread. This can avoid deadlocks in cases where background tasks end up waiting on I/O and we don't have enough background task threads to fulfill decompression requests. This is off by default

#tests tested locally, also perf tested via replayrun exe override preflight in horde

#localization none

diff --git a/Engine/Source/Runtime/PakFile/Private/IoDispatcherFileBackend.cpp b/Engine/Source/Runtime/PakFile/Private/IoDispatcherFileBackend.cpp
--- a/Engine/Source/Runtime/PakFile/Private/IoDispatcherFileBackend.cpp
+++ b/Engine/Source/Runtime/PakFile/Private/IoDispatcherFileBackend.cpp
@@ -97,6 +97,15 @@ static FAutoConsoleVariableRef CVar_IoDispatcherTocsEnablePerfectHashing(
 	TEXT("Enable perfect hashmap lookups for iostore tocs")
 );
 
+int32 GIoDispatcherForceSynchronousScatter = 0;
+static FAutoConsoleVariableRef CVar_IoDispatcherForceSynchronousScatter(
+	TEXT("s.IoDispatcherForceSynchronousScatter"),
+	GIoDispatcherForceSynchronousScatter,
+	TEXT("Force scatter jobs to be synchronous on the IODispatcher thread.\n")
+	TEXT("This can avoid deadlocks in cases where background tasks end up waiting on I/O and we don't have enough background task threads to fulfill decompression requests.")
+);
+
+
 uint32 FFileIoStoreReadRequest::NextSequence = 0;
 #if CHECK_IO_STORE_READ_REQUEST_LIST_MEMBERSHIP
 uint32 FFileIoStoreReadRequestList::NextListCookie = 0;
@@ -1826,7 +1835,7 @@ FIoRequestImpl* FFileIoStore::GetCompletedRequests()
 		}
 
 		// Scatter block asynchronous when the block is compressed, encrypted or signed
-		const bool bScatterAsync = bIsMultithreaded && (!BlockToDecompress->CompressionMethod.IsNone() || BlockToDecompress->EncryptionKey.IsValid() || BlockToDecompress->SignatureHash);
+		bool bScatterAsync = bIsMultithreaded && GIoDispatcherForceSynchronousScatter==0 && (!BlockToDecompress->CompressionMethod.IsNone() || BlockToDecompress->EncryptionKey.IsValid() || BlockToDecompress->SignatureHash);
 		if (bScatterAsync)
 		{
 			TGraphTask<FDecompressAsyncTask>::CreateTask().ConstructAndDispatchWhenReady(*this, BlockToDecompress);
@@ -2451,4 +2460,4 @@ void FFileIoStoreStats::OnBufferAllocated()
 #endif
 }
 
-#endif
+#endif
\ No newline at end of file
