PAK - remove internal encryption key cache

diff --git a/Engine/Source/Runtime/PakFile/Private/IPlatformFilePak.cpp b/Engine/Source/Runtime/PakFile/Private/IPlatformFilePak.cpp
--- a/Engine/Source/Runtime/PakFile/Private/IPlatformFilePak.cpp
+++ b/Engine/Source/Runtime/PakFile/Private/IPlatformFilePak.cpp
@@ -32,6 +32,7 @@
 #include "Misc/ConfigCacheIni.h"
 #endif
 #include "ProfilingDebugging/CsvProfiler.h"
+#include "Misc/EncryptionKeyManager.h"
 #include "Misc/Fnv.h"
 #include "ProfilingDebugging/ScopedTimers.h"
 #include "Async/MappedFileHandle.h"
@@ -91,53 +92,6 @@ int32 GetPakchunkIndexFromPakFile(const FString& InFilename)
 	return FGenericPlatformMisc::GetPakchunkIndexFromPakFile(InFilename);
 }
 
-// Registered encryption key cache
-class FEncryptionKeyCache
-{
-public:
-
-	void AddKey(const FGuid& InGuid, const FAES::FAESKey InKey)
-	{
-		FScopeLock Lock(&SyncObject);
-		if (!Keys.Contains(InGuid))
-		{
-			Keys.Add(InGuid, InKey);
-		}
-	}
-
-	bool GetKey(const FGuid& InGuid, FAES::FAESKey& OutKey)
-	{
-		FScopeLock Lock(&SyncObject);
-		if (const FAES::FAESKey* Key = Keys.Find(InGuid))
-		{
-			OutKey = *Key;
-			return true;
-		}
-		return false;
-	}
-
-	bool const HasKey(const FGuid& InGuid)
-	{
-		return Keys.Contains(InGuid);
-	}
-
-	const TMap<FGuid, FAES::FAESKey>& GetKeys() const
-	{
-		return Keys;
-	}
-
-private:
-
-	TMap<FGuid, FAES::FAESKey> Keys;
-	FCriticalSection SyncObject;
-};
-
-FEncryptionKeyCache& GetRegisteredEncryptionKeys()
-{
-	static FEncryptionKeyCache Instance;
-	return Instance;
-}
-
 #if !UE_BUILD_SHIPPING
 static void TestRegisterEncryptionKey(const TArray<FString>& Args)
 {
@@ -248,13 +202,16 @@ void FPakPlatformFile::GetFilenamesFromIostoreByBlockIndex(const FString& InCont
 	{
 		return;
 	}
+
+	const TMap<FGuid, FAES::FAESKey> Keys = UE::FEncryptionKeyManager::Get().GetAllKeys();
+
 	FScopeLock ScopedLock(&PakPlatformFile->PakListCritical);
 	for (const FPakListEntry& PakListEntry : PakPlatformFile->PakFiles)
 	{
 		if (FPaths::GetBaseFilename(PakListEntry.PakFile->PakFilename) == InContainerName)
 		{
 			TUniquePtr<FIoStoreReader> IoStoreReader(new FIoStoreReader());
-			FIoStatus Status = IoStoreReader->Initialize(*FPaths::ChangeExtension(PakListEntry.PakFile->PakFilename, TEXT("")), GetRegisteredEncryptionKeys().GetKeys());
+			FIoStatus Status = IoStoreReader->Initialize(*FPaths::ChangeExtension(PakListEntry.PakFile->PakFilename, TEXT("")),  Keys);
 			if (Status.IsOk())
 			{
 				IoStoreReader->GetFilenamesByBlockIndex(InBlockIndex, OutFileList);
@@ -818,13 +775,16 @@ void FPakPlatformFile::GetFilenamesFromIostoreContainer(const FString& InContain
 	{
 		return;
 	}
+
+	const TMap<FGuid, FAES::FAESKey> Keys = UE::FEncryptionKeyManager::Get().GetAllKeys();
+
 	FScopeLock ScopedLock(&PakPlatformFile->PakListCritical);
 	for (const FPakListEntry& PakListEntry : PakPlatformFile->PakFiles)
 	{
 		if (FPaths::GetBaseFilename(PakListEntry.PakFile->PakFilename) == InContainerName)
 		{
 			TUniquePtr<FIoStoreReader> IoStoreReader(new FIoStoreReader());
-			FIoStatus Status = IoStoreReader->Initialize(*FPaths::ChangeExtension(PakListEntry.PakFile->PakFilename, TEXT("")), GetRegisteredEncryptionKeys().GetKeys());
+			FIoStatus Status = IoStoreReader->Initialize(*FPaths::ChangeExtension(PakListEntry.PakFile->PakFilename, TEXT("")), Keys);
 			if (Status.IsOk())
 			{
 				IoStoreReader->GetFilenames(OutFileList);
@@ -841,11 +801,14 @@ void FPakPlatformFile::ForeachPackageInIostoreWhile(TFunctionRef<bool(FName)> Pr
 	{
 		return;
 	}
+
+	const TMap<FGuid, FAES::FAESKey> Keys = UE::FEncryptionKeyManager::Get().GetAllKeys();
+
 	FScopeLock ScopedLock(&PakPlatformFile->PakListCritical);
 	for (const FPakListEntry& PakListEntry : PakPlatformFile->PakFiles)
 	{
 		TUniquePtr<FIoStoreReader> IoStoreReader(new FIoStoreReader());
-		FIoStatus Status = IoStoreReader->Initialize(*FPaths::ChangeExtension(PakListEntry.PakFile->PakFilename, TEXT("")), GetRegisteredEncryptionKeys().GetKeys());
+		FIoStatus Status = IoStoreReader->Initialize(*FPaths::ChangeExtension(PakListEntry.PakFile->PakFilename, TEXT("")), Keys);
 		if (Status.IsOk())
 		{
 			const FIoDirectoryIndexReader& DirectoryIndex = IoStoreReader->GetDirectoryIndexReader();
@@ -891,7 +854,7 @@ void FPakPlatformFile::GetPakEncryptionKey(FAES::FAESKey& OutKey, const FGuid& I
 {
 	OutKey.Reset();
 
-	if (!GetRegisteredEncryptionKeys().GetKey(InEncryptionKeyGuid, OutKey))
+	if (!UE::FEncryptionKeyManager::Get().TryGetKey(InEncryptionKeyGuid, OutKey))
 	{
 		if (!InEncryptionKeyGuid.IsValid() && FCoreDelegates::GetPakEncryptionKeyDelegate().IsBound())
 		{
@@ -6046,7 +6009,7 @@ void FPakFile::Initialize(FArchive& Reader, bool bLoadIndex)
 		UE_CLOG(!((Info.IndexOffset + Info.IndexSize) >= 0 && (Info.IndexOffset + Info.IndexSize) <= CachedTotalSize), LogPakFile, Fatal, TEXT("Index end offset for pak file '%s' is invalid (%lld)"), *PakFilename, Info.IndexOffset + Info.IndexSize);
 
 		// If we aren't using a dynamic encryption key, process the pak file using the embedded key
-		if (!Info.EncryptionKeyGuid.IsValid() || GetRegisteredEncryptionKeys().HasKey(Info.EncryptionKeyGuid))
+		if (!Info.EncryptionKeyGuid.IsValid() || UE::FEncryptionKeyManager::Get().ContainsKey(Info.EncryptionKeyGuid))
 		{
 			if (bLoadIndex)
 			{
@@ -8061,7 +8024,7 @@ FPakPlatformFile::FPakPlatformFile()
 	: LowerLevel(NULL)
 	, bSigned(false)
 {
-	FCoreDelegates::GetRegisterEncryptionKeyMulticastDelegate().AddRaw(this, &FPakPlatformFile::RegisterEncryptionKey);
+	UE::FEncryptionKeyManager::Get().OnKeyAdded().AddRaw(this, &FPakPlatformFile::RegisterEncryptionKey);
 }
 
 FPakPlatformFile::~FPakPlatformFile()
@@ -8070,7 +8033,7 @@ FPakPlatformFile::~FPakPlatformFile()
 
 	FTSTicker::GetCoreTicker().RemoveTicker(RetireReadersHandle);
 
-	FCoreDelegates::GetRegisterEncryptionKeyMulticastDelegate().RemoveAll(this);
+	UE::FEncryptionKeyManager::Get().OnKeyAdded().RemoveAll(this);
 	FCoreDelegates::OnFEngineLoopInitComplete.RemoveAll(this);
 
 	FCoreDelegates::OnMountAllPakFiles.Unbind();
@@ -8483,7 +8446,7 @@ bool FPakPlatformFile::Mount(const TCHAR* InPakFilename, uint32 PakOrder, const
 		TRefCountPtr<FPakFile> Pak = new FPakFile(LowerLevel, InPakFilename, bSigned, bLoadIndex);
 		if (Pak.GetReference()->IsValid())
 		{
-			if (!Pak->GetInfo().EncryptionKeyGuid.IsValid() || GetRegisteredEncryptionKeys().HasKey(Pak->GetInfo().EncryptionKeyGuid))
+			if (!Pak->GetInfo().EncryptionKeyGuid.IsValid() || UE::FEncryptionKeyManager::Get().ContainsKey(Pak->GetInfo().EncryptionKeyGuid))
 			{
 				if (InPath != nullptr)
 				{
@@ -8538,7 +8501,7 @@ bool FPakPlatformFile::Mount(const TCHAR* InPakFilename, uint32 PakOrder, const
 			{
 				UE_LOG(LogPakFile, Display, TEXT("Deferring mount of pak \"%s\" until encryption key '%s' becomes available"), InPakFilename, *Pak->GetInfo().EncryptionKeyGuid.ToString());
 
-				check(!GetRegisteredEncryptionKeys().HasKey(Pak->GetInfo().EncryptionKeyGuid));
+				check(!UE::FEncryptionKeyManager::Get().ContainsKey(Pak->GetInfo().EncryptionKeyGuid));
 				FPakListDeferredEntry& Entry = PendingEncryptedPakFiles[PendingEncryptedPakFiles.Add(FPakListDeferredEntry())];
 				Entry.Filename = InPakFilename;
 				Entry.Path = InPath;
@@ -8560,7 +8523,7 @@ bool FPakPlatformFile::Mount(const TCHAR* InPakFilename, uint32 PakOrder, const
 			FGuid EncryptionKeyGuid = Pak->GetInfo().EncryptionKeyGuid;
 			FAES::FAESKey EncryptionKey;
 
-			if (!GetRegisteredEncryptionKeys().GetKey(EncryptionKeyGuid, EncryptionKey))
+			if (!UE::FEncryptionKeyManager::Get().TryGetKey(EncryptionKeyGuid, EncryptionKey))
 			{
 				if (!EncryptionKeyGuid.IsValid() && FCoreDelegates::GetPakEncryptionKeyDelegate().IsBound())
 				{
@@ -8896,8 +8859,6 @@ bool FPakPlatformFile::HandleUnmountPakDelegate(const FString& PakFilePath)
 
 void FPakPlatformFile::RegisterEncryptionKey(const FGuid& InGuid, const FAES::FAESKey& InKey)
 {
-	GetRegisteredEncryptionKeys().AddKey(InGuid, InKey);
-
 	int32 NumMounted = 0;
 	TSet<int32> ChunksToNotify;
 
