Add new delegates to better allow various IoDispatcher subsystems to mount their containers when .pak files are mounted.

#rnx

- The filebased iostore system is tightly linked to the existing pakfile code and recieved notifications when a PakFile is mounted/unmounted directly in code. However we do not want to force the IAS system on all users, it should be opt in. So we need a generic way to hook into the PakFile system and be notified about PakFile operations.
- Add a new delegate (FCoreInternalDelegates::GetOnPakMountOperation) that is fired when ever a pakfile is mounted or unmounted.
- Add a new delegate (FCoreInternalDelegates::GetCurrentlyMountedPaksDelegate) that can be queried to find all of the pakfiles that are currently mounted.
- Both delegates are in a new internal header file to limit their use to the engine which will make it easier to change in the future if we need to do so.

diff --git a/Engine/Source/Runtime/Core/Internal/Misc/CoreDelegatesInternal.h b/Engine/Source/Runtime/Core/Internal/Misc/CoreDelegatesInternal.h
new file mode 100644
--- /dev/null
+++ b/Engine/Source/Runtime/Core/Internal/Misc/CoreDelegatesInternal.h
@@ -0,0 +1,44 @@
+// Copyright Epic Games, Inc. All Rights Reserved.
+
+#pragma once
+
+// This header is used for Core delegates that we don't want to expose outside of the engine
+
+#include "Delegates/Delegate.h"
+#include "Misc/AES.h"
+#include "Misc/Guid.h"
+
+class IPakFile;
+
+enum class EMountOperation
+{
+	Mount,
+	Unmount
+};
+
+/** Delegate to indicate that a pakfile has been mounted/unmounted */
+DECLARE_TS_MULTICAST_DELEGATE_ThreeParams(FMountOperationPak, EMountOperation /*Operation*/, const TCHAR* /*PakPath*/, int32 /*Order*/);
+
+struct FMountedPakInfo
+{
+	FMountedPakInfo() = default;
+	FMountedPakInfo(const IPakFile* InPakFile, int32 InOrder)
+		: PakFile(InPakFile)
+		, Order(InOrder)
+	{}
+
+	~FMountedPakInfo() = default;
+
+	const IPakFile* PakFile = nullptr;
+	int32 Order = INDEX_NONE;
+};
+
+FUNC_DECLARE_DELEGATE(FCurrentlyMountedPaksDelegate, TArray<FMountedPakInfo>);
+
+struct FCoreInternalDelegates
+{
+	// Callback to prompt the pak system to unmount a pak file.
+	static CORE_API FMountOperationPak& GetOnPakMountOperation();
+
+	static CORE_API FCurrentlyMountedPaksDelegate& GetCurrentlyMountedPaksDelegate();
+};
diff --git a/Engine/Source/Runtime/Core/Private/Misc/CoreDelegates.cpp b/Engine/Source/Runtime/Core/Private/Misc/CoreDelegates.cpp
--- a/Engine/Source/Runtime/Core/Private/Misc/CoreDelegates.cpp
+++ b/Engine/Source/Runtime/Core/Private/Misc/CoreDelegates.cpp
@@ -1,10 +1,10 @@
 // Copyright Epic Games, Inc. All Rights Reserved.
 
-// Core includes.
 #include "Misc/CoreDelegates.h"
-#include "Math/Vector.h"
-#include "Misc/Fork.h"
 
+#include "Misc/CoreDelegatesInternal.h"
+#include "Misc/Fork.h"
+#include "Math/Vector.h"
 
 //////////////////////////////////////////////////////////////////////////
 // FCoreDelegates
@@ -55,6 +55,23 @@ TTSMulticastDelegate<void(const TCHAR*, const TCHAR*)>& FCoreDelegates::GetOnFil
 	return Singleton;
 }
 
+//namespace UE::Core
+//{
+
+FMountOperationPak& FCoreInternalDelegates::GetOnPakMountOperation()
+{
+	static FMountOperationPak Delegate;
+	return Delegate;
+}
+
+FCurrentlyMountedPaksDelegate& FCoreInternalDelegates::GetCurrentlyMountedPaksDelegate()
+{
+	static FCurrentlyMountedPaksDelegate Delegate;
+	return Delegate;
+}
+
+//} // namespace UE::Core
+
 TMulticastDelegate<void(bool, int32, int32)> FCoreDelegates::OnUserLoginChangedEvent;
 PRAGMA_DISABLE_DEPRECATION_WARNINGS
 TMulticastDelegate<void(bool, FPlatformUserId, int32)> FCoreDelegates::OnControllerConnectionChange;
diff --git a/Engine/Source/Runtime/PakFile/Private/IPlatformFilePak.cpp b/Engine/Source/Runtime/PakFile/Private/IPlatformFilePak.cpp
--- a/Engine/Source/Runtime/PakFile/Private/IPlatformFilePak.cpp
+++ b/Engine/Source/Runtime/PakFile/Private/IPlatformFilePak.cpp
@@ -10,6 +10,7 @@
 #include "HAL/IConsoleManager.h"
 #include "HAL/LowLevelMemTracker.h"
 #include "Misc/CoreDelegates.h"
+#include "Misc/CoreDelegatesInternal.h"
 #include "Misc/App.h"
 #include "Modules/ModuleManager.h"
 #include "Misc/ConfigCacheIni.h"
@@ -8283,6 +8284,21 @@ bool FPakPlatformFile::Initialize(IPlatformFile* Inner, const TCHAR* CmdLine)
 	FCoreDelegates::MountPak.BindRaw(this, &FPakPlatformFile::HandleMountPakDelegate);
 	FCoreDelegates::OnUnmountPak.BindRaw(this, &FPakPlatformFile::HandleUnmountPakDelegate);
 	FCoreDelegates::OnOptimizeMemoryUsageForMountedPaks.BindRaw(this, &FPakPlatformFile::OptimizeMemoryUsageForMountedPaks);
+	FCoreInternalDelegates::GetCurrentlyMountedPaksDelegate().BindLambda([this]()
+		{
+			TArray<FPakListEntry> Paks;
+			GetMountedPaks(Paks);
+
+			TArray<FMountedPakInfo> PakInfo;
+			PakInfo.Reserve(Paks.Num());
+			
+			for (const FPakListEntry& Entry : Paks)
+			{
+				PakInfo.Emplace(FMountedPakInfo(Entry.PakFile, Entry.ReadOrder));
+			}
+
+			return PakInfo;
+		});
 
 	FCoreDelegates::OnFEngineLoopInitComplete.AddRaw(this, &FPakPlatformFile::OptimizeMemoryUsageForMountedPaks);
 
@@ -8575,6 +8591,11 @@ bool FPakPlatformFile::Mount(const TCHAR* InPakFilename, uint32 PakOrder, const
 			}
 		}
 
+		if (bPakSuccess && FCoreInternalDelegates::GetOnPakMountOperation().IsBound())
+		{
+			FCoreInternalDelegates::GetOnPakMountOperation().Broadcast(EMountOperation::Mount, InPakFilename, PakOrder);
+		}
+
 		if (bPakSuccess)
 		{
 			double OnPakFileMounted2Time = 0.0;
