PakFiles: deprecated OnPakFileMounted2: Prevent calling Broadcast if unsubscribed, to prevent the invalid-simultaneous-writeaccess detector from spuriously firing.
#rnx

diff --git a/Engine/Source/Runtime/PakFile/Private/IPlatformFilePak.cpp b/Engine/Source/Runtime/PakFile/Private/IPlatformFilePak.cpp
--- a/Engine/Source/Runtime/PakFile/Private/IPlatformFilePak.cpp
+++ b/Engine/Source/Runtime/PakFile/Private/IPlatformFilePak.cpp
@@ -8603,7 +8603,12 @@ bool FPakPlatformFile::Mount(const TCHAR* InPakFilename, uint32 PakOrder, const
 				FScopedDurationTimer Timer(OnPakFileMounted2Time);
 
 PRAGMA_DISABLE_DEPRECATION_WARNINGS
-				FCoreDelegates::OnPakFileMounted2.Broadcast(*Pak);
+				if (FCoreDelegates::OnPakFileMounted2.IsBound())
+				{
+					// Avoid calling Broadcast if not in use; Broadcast even on an unsubscribed
+					// non-threadsafe delegate is not threadsafe.
+					FCoreDelegates::OnPakFileMounted2.Broadcast(*Pak);
+				}
 PRAGMA_ENABLE_DEPRECATION_WARNINGS
 
 				FCoreDelegates::GetOnPakFileMounted2().Broadcast(*Pak);
