PakFiles: Suppress spurious log message "Mount point '%s' is not mounted to a valid Root Path yet" when the mount point is a parent path of existing root paths.
#rnx

diff --git a/Engine/Source/Runtime/CoreUObject/Private/Misc/PackageName.cpp b/Engine/Source/Runtime/CoreUObject/Private/Misc/PackageName.cpp
--- a/Engine/Source/Runtime/CoreUObject/Private/Misc/PackageName.cpp
+++ b/Engine/Source/Runtime/CoreUObject/Private/Misc/PackageName.cpp
@@ -2506,6 +2506,25 @@ void FPackageName::QueryRootContentPaths(TArray<FString>& OutRootContentPaths, b
 	}
 }
 
+TArray<FString> FPackageName::QueryMountPointLocalAbsPaths()
+{
+	const FLongPackagePathsSingleton& Paths = FLongPackagePathsSingleton::Get();
+	TArray<FString> OutAbsPaths;
+
+	{
+		FReadScopeLock ScopeLock(Paths.MountLock);
+		OutAbsPaths.Reserve(Paths.MountPoints.Num() + 1);
+		for (const TUniquePtr<FMountPoint>& MountPoint : Paths.MountPoints)
+		{
+			if (!MountPoint->bAlias && !MountPoint->ContentPathAbsolute.IsEmpty())
+			{
+				OutAbsPaths.Add(MountPoint->ContentPathAbsolute);
+			}
+		}
+	}
+	return OutAbsPaths;
+}
+
 void FPackageName::OnCoreUObjectInitialized()
 {
 	FLongPackagePathsSingleton::Get().OnCoreUObjectInitialized();
diff --git a/Engine/Source/Runtime/CoreUObject/Public/Misc/PackageName.h b/Engine/Source/Runtime/CoreUObject/Public/Misc/PackageName.h
--- a/Engine/Source/Runtime/CoreUObject/Public/Misc/PackageName.h
+++ b/Engine/Source/Runtime/CoreUObject/Public/Misc/PackageName.h
@@ -626,6 +626,12 @@ public:
 	 */
 	static COREUOBJECT_API void QueryRootContentPaths( TArray<FString>& OutRootContentPaths, bool bIncludeReadOnlyRoots = false, bool bWithoutLeadingSlashes = false, bool bWithoutTrailingSlashes = false);
 	
+	/**
+	 * Returns all of the local paths on disk of the root content paths, like
+	 * "c:\MyProjects\ProjectA\Content", "d:\Unreal\Engine\Content", "c:\MyProjects\ProjectA\Plugins\MyPlugin\Content"
+	 */
+	static COREUOBJECT_API TArray<FString> QueryMountPointLocalAbsPaths();
+
 	/** If the FLongPackagePathsSingleton is not created yet, this function will create it and thus allow mount points to be added */
 	static COREUOBJECT_API void OnCoreUObjectInitialized();
 
diff --git a/Engine/Source/Runtime/PakFile/Private/IPlatformFilePak.cpp b/Engine/Source/Runtime/PakFile/Private/IPlatformFilePak.cpp
--- a/Engine/Source/Runtime/PakFile/Private/IPlatformFilePak.cpp
+++ b/Engine/Source/Runtime/PakFile/Private/IPlatformFilePak.cpp
@@ -8617,7 +8617,7 @@ PRAGMA_ENABLE_DEPRECATION_WARNINGS
 			UE_LOG(LogPakFile, Display, TEXT("Mounted Pak file '%s', mount point: '%s'"), InPakFilename, *Pak->GetMountPoint());
 			UE_LOG(LogPakFile, Verbose, TEXT("OnPakFileMounted2Time == %lf"), OnPakFileMounted2Time);
 						
-			// skip this check for the default mountpoint, it will print false positives
+			// skip this check for the default mountpoint, it is a frequently used known-good mount point
 			FString NormalizedPakMountPoint = FPaths::CreateStandardFilename(Pak->GetMountPoint());
 			bool bIsMountingToRoot = NormalizedPakMountPoint == FPaths::CreateStandardFilename(FPaths::RootDir());
 #if WITH_EDITOR
@@ -8626,9 +8626,29 @@ PRAGMA_ENABLE_DEPRECATION_WARNINGS
 			if (!bIsMountingToRoot)
 			{
 				FString OutPackageName;
-				if (!FPackageName::TryConvertFilenameToLongPackageName(Pak->GetMountPoint(), OutPackageName))
+				const FString& MountPoint = Pak->GetMountPoint();
+				if (!FPackageName::TryConvertFilenameToLongPackageName(MountPoint, OutPackageName))
 				{
-					UE_LOG(LogPakFile, Display, TEXT("Mount point: '%s' is not mounted to a valid Root Path yet, assets in this pak file may not be accessible until a corresponding UFS Mount Point is added through FPackageName::RegisterMountPoint."), *Pak->GetMountPoint());
+					// Possibly the mount point is a parent path of mount points, e.g. <ProjectRoot>/Plugins,
+					// parent path of <ProjectRoot>/Plugins/PluginA and <ProjectRoot>/Plugins/PluginB.
+					// Do not warn in that case.
+					FString MountPointAbsPath = FPaths::ConvertRelativePathToFull(MountPoint);
+					bool bParentOfMountPoint = false;
+					for (const FString& ExistingMountPoint : FPackageName::QueryMountPointLocalAbsPaths())
+					{
+						if (FPathViews::IsParentPathOf(MountPointAbsPath, ExistingMountPoint))
+						{
+							bParentOfMountPoint = true;
+							break;
+						}
+					}
+					if (!bParentOfMountPoint)
+					{
+						UE_LOG(LogPakFile, Display,
+							TEXT("Mount point '%s' is not mounted to a valid Root Path yet, ")
+							TEXT("assets in this pak file may not be accessible until a corresponding UFS Mount Point is added through FPackageName::RegisterMountPoint."),
+							*MountPoint);
+					}
 				}
 			}
 		}
