PakFiles: DeleteRecords: Allow files that are duplicated between chunks to still be found if one chunk deleted them in a patch but the other chunk did not.
#rnx

diff --git a/Engine/Source/Runtime/PakFile/Private/IPlatformFilePak.cpp b/Engine/Source/Runtime/PakFile/Private/IPlatformFilePak.cpp
--- a/Engine/Source/Runtime/PakFile/Private/IPlatformFilePak.cpp
+++ b/Engine/Source/Runtime/PakFile/Private/IPlatformFilePak.cpp
@@ -47,6 +47,7 @@
 #include "Compression/OodleDataCompression.h"
 #include "IO/IoStore.h"
 #include "String/RemoveFrom.h"
+#include "Algo/AnyOf.h"
 
 DEFINE_LOG_CATEGORY(LogPakFile);
 
@@ -248,42 +249,60 @@ bool FPakPlatformFile::FindFileInPakFiles(TArray<FPakListEntry>& Paks, const TCH
 	FString StandardFilename(Filename);
 	FPaths::MakeStandardFilename(StandardFilename);
 
-	int32 DeletedReadOrder = -1;
+	TArray<const FPakListEntry*, TInlineAllocator<1>> PaksWithDeleteRecord;
+	bool bFoundOlderVersionOfDeleteRecordPak = false;
 
 	for (int32 PakIndex = 0; PakIndex < Paks.Num(); PakIndex++)
 	{
-		int32 PakReadOrder = Paks[PakIndex].ReadOrder;
-		if (DeletedReadOrder != -1 && DeletedReadOrder > PakReadOrder)
+		const FPakListEntry& PakEntry = Paks[PakIndex];
+		FPakFile* PakFile = PakEntry.PakFile.GetReference();
+		if (!PakFile)
 		{
-			// Found a delete record in a higher priority patch level, but now we're at a lower priority set.
-			// Don't search further back or we'll find the original, old file.
-			UE_LOG(LogPakFile, Verbose, TEXT("Delete Record: Accepted a delete record for %s"), Filename);
-			return false;
+			continue;
+		}
+
+		if (PaksWithDeleteRecord.Num() > 0)
+		{
+			if (Algo::AnyOf(PaksWithDeleteRecord, [&PakEntry](const FPakListEntry* DeletedPakEntry)
+				{
+					return DeletedPakEntry->ReadOrder > PakEntry.ReadOrder &&
+						DeletedPakEntry->PakFile->PakchunkIndex == PakEntry.PakFile->PakchunkIndex;
+				}))
+			{
+				// Found a delete record in a higher priority patch level, and this is an earlier version of the same file.
+				// Don't search in the file.
+				bFoundOlderVersionOfDeleteRecordPak = true;
+				continue;
+			}
 		}
 
-		FPakFile::EFindResult FindResult = Paks[PakIndex].PakFile->Find(StandardFilename, OutEntry);
+		FPakFile::EFindResult FindResult = PakFile->Find(StandardFilename, OutEntry);
 		if (FindResult == FPakFile::EFindResult::Found)
 		{
 			if (OutPakFile != NULL)
 			{
-				*OutPakFile = Paks[PakIndex].PakFile;
+				*OutPakFile = PakFile;
 			}
-			UE_CLOG(DeletedReadOrder != -1, LogPakFile, Verbose,
-				TEXT("Delete Record: Ignored delete record for %s - found it in %s instead (asset was moved between chunks)"),
-				Filename, *Paks[PakIndex].PakFile->GetFilename());
+			UE_CLOG(!PaksWithDeleteRecord.IsEmpty(), LogPakFile, Verbose,
+				TEXT("Delete Record: Ignored delete record for %s - found it in %s instead (asset was moved or duplicated between chunks)"),
+				Filename, *PakFile->GetFilename());
 			return true;
 		}
 		else if (FindResult == FPakFile::EFindResult::FoundDeleted)
 		{
-			DeletedReadOrder = PakReadOrder;
+			PaksWithDeleteRecord.Add(&PakEntry);
 			UE_LOG(LogPakFile, Verbose, TEXT("Delete Record: Found a delete record for %s in %s"),
-				Filename, *Paks[PakIndex].PakFile->GetFilename());
+				Filename, *PakFile->GetFilename());
 		}
 	}
 
-	UE_CLOG(DeletedReadOrder != -1, LogPakFile, Warning,
-		TEXT("Delete Record: No lower priority pak files looking for %s. (maybe not downloaded?)"),
-		Filename);
+	if (!PaksWithDeleteRecord.IsEmpty())
+	{
+		UE_CLOG(bFoundOlderVersionOfDeleteRecordPak, LogPakFile, Verbose,
+			TEXT("Delete Record: Accepted a delete record for %s"), Filename);
+		UE_CLOG(!bFoundOlderVersionOfDeleteRecordPak, LogPakFile, Warning,
+			TEXT("Delete Record: No lower priority pak files looking for %s. (maybe not downloaded?)"), Filename);
+	}
 	return false;
 }
 
