ZenLoader: Extend container mount logs to get some more information on errors of type "Failed reading chunk for package:  (Not Found)" when there is no ExportBundleData chunk for a package known by the loader.
Add LexToString(const FIoContainerId& ContainerId).

#rnx

diff --git a/Engine/Source/Runtime/Core/Private/IO/IoContainerId.cpp b/Engine/Source/Runtime/Core/Private/IO/IoContainerId.cpp
--- a/Engine/Source/Runtime/Core/Private/IO/IoContainerId.cpp
+++ b/Engine/Source/Runtime/Core/Private/IO/IoContainerId.cpp
@@ -32,3 +32,8 @@ void operator<<(FStructuredArchiveSlot Slot, FIoContainerId& Value)
 {
 	Slot << Value.Id;
 }
+
+FString LexToString(const FIoContainerId& ContainerId)
+{
+	return FString::Printf(TEXT("%llx"), ContainerId.Id);
+}
diff --git a/Engine/Source/Runtime/Core/Public/IO/IoContainerId.h b/Engine/Source/Runtime/Core/Public/IO/IoContainerId.h
--- a/Engine/Source/Runtime/Core/Public/IO/IoContainerId.h
+++ b/Engine/Source/Runtime/Core/Public/IO/IoContainerId.h
@@ -56,6 +56,8 @@ public:
 
 	CORE_API friend void operator<<(FStructuredArchiveSlot Slot, FIoContainerId& Value);
 
+	CORE_API friend FString LexToString(const FIoContainerId& ContainerId);
+
 private:
 	inline explicit FIoContainerId(const uint64 InId)
 		: Id(InId) { }
diff --git a/Engine/Source/Runtime/PakFile/Private/FilePackageStore.cpp b/Engine/Source/Runtime/PakFile/Private/FilePackageStore.cpp
--- a/Engine/Source/Runtime/PakFile/Private/FilePackageStore.cpp
+++ b/Engine/Source/Runtime/PakFile/Private/FilePackageStore.cpp
@@ -399,6 +399,9 @@ void FFilePackageStoreBackend::Mount(FIoContainerHeader* ContainerHeader, uint32
 			return A.Order > B.Order;
 		});
 	bNeedsContainerUpdate = true;
+
+	UE_LOG(LogFilePackageStore, Log, TEXT("Mounting container: Id=%s, Order=%u, NumPackages=%d"),
+		*LexToString(ContainerHeader->ContainerId), Order, ContainerHeader->PackageIds.Num());
 }
 
 void FFilePackageStoreBackend::Unmount(const FIoContainerHeader* ContainerHeader)
@@ -704,6 +707,9 @@ void FFilePackageStoreBackend::Update()
 		PackageEntries = FPackageIdMap(MoveTemp(Pairs));
 
 		bNeedsContainerUpdate = false;	
+
+		UE_LOG(LogFilePackageStore, Log, TEXT("Updated: NewPackages=%u, OldPackages=%u, TotalPackages=%u"),
+			TotalNewPackages, TotalOldPackages, PackageEntries.GetCapacity());
 	}
 }
 
@@ -830,4 +836,4 @@ TEST_CASE_NAMED(FPackageIdMapTest, "System::Pak::IoStore::PackageIdMap", "[Pak][
 	}
 }
 
-#endif
\ No newline at end of file
+#endif
diff --git a/Engine/Source/Runtime/PakFile/Private/IoDispatcherFileBackend.cpp b/Engine/Source/Runtime/PakFile/Private/IoDispatcherFileBackend.cpp
--- a/Engine/Source/Runtime/PakFile/Private/IoDispatcherFileBackend.cpp
+++ b/Engine/Source/Runtime/PakFile/Private/IoDispatcherFileBackend.cpp
@@ -743,7 +743,8 @@ FIoStatus FFileIoStoreReader::Initialize(const TCHAR* InTocFilePath, int32 InOrd
 
 	Stats.OnTocMounted(GetTocAllocatedSize());
 
-	UE_LOG(LogIoDispatcher, Display, TEXT("Toc signature hash: %s"), *TocResource.SignatureHash.ToString());
+	UE_LOG(LogIoDispatcher, Display, TEXT("Toc loaded : %s, Id=%s, Order=%d, EntryCount=%u, SignatureHash=%s"),
+		InTocFilePath, *LexToString(TocResource.Header.ContainerId), InOrder, TocResource.Header.TocEntryCount, *TocResource.SignatureHash.ToString());
 	
 	ContainerId = TocResource.Header.ContainerId;
 	Order = InOrder;
