[Trace] Fix CIS Add TraceLog module to linker

diff --git a/Engine/Source/Runtime/PakFile/PakFile.Build.cs b/Engine/Source/Runtime/PakFile/PakFile.Build.cs
--- a/Engine/Source/Runtime/PakFile/PakFile.Build.cs
+++ b/Engine/Source/Runtime/PakFile/PakFile.Build.cs
@@ -9,6 +9,7 @@ public class PakFile : ModuleRules
 		UnsafeTypeCastWarningLevel = WarningLevel.Error;
 		PrivateDependencyModuleNames.Add("Core");
 		PrivateDependencyModuleNames.Add("CoreUObject");
+		PrivateDependencyModuleNames.Add("TraceLog");
 		PublicDependencyModuleNames.Add("RSA");
 
 		bAllowAutoRTFMInstrumentation = true;
