IoDispatcher: Use the new Tasks API to launch scatter block decompression tasks.
Reduce risk of deadlock if other background worker tasks are waiting for io by allowing decompression tasks to run on both foreground and background workers.
Respect the existing priority cvar "TaskGraph.TaskPriorities.IoDispatcherAsyncTasks".

#tests pc and console replays

diff --git a/Engine/Source/Runtime/PakFile/Private/IoDispatcherFileBackend.cpp b/Engine/Source/Runtime/PakFile/Private/IoDispatcherFileBackend.cpp
--- a/Engine/Source/Runtime/PakFile/Private/IoDispatcherFileBackend.cpp
+++ b/Engine/Source/Runtime/PakFile/Private/IoDispatcherFileBackend.cpp
@@ -1514,11 +1514,6 @@ FAutoConsoleTaskPriority CPrio_IoDispatcherTaskPriority(
 	ENamedThreads::NormalTaskPriority // if we don't have background threads, then use normal priority threads at normal task priority instead
 );
 
-ENamedThreads::Type FFileIoStore::FDecompressAsyncTask::GetDesiredThread()
-{
-	return CPrio_IoDispatcherTaskPriority.Get();
-}
-
 void FFileIoStore::ScatterBlock(FFileIoStoreCompressedBlock* CompressedBlock, bool bIsAsync)
 {
 	LLM_SCOPE(ELLMTag::FileSystem);
@@ -1851,10 +1846,22 @@ FIoRequestImpl* FFileIoStore::GetCompletedIoRequests()
 		}
 
 		// Scatter block asynchronous when the block is compressed, encrypted or signed
-		bool bScatterAsync = bIsMultithreaded && GIoDispatcherForceSynchronousScatter==0 && (!BlockToDecompress->CompressionMethod.IsNone() || BlockToDecompress->EncryptionKey.IsValid() || BlockToDecompress->SignatureHash);
+		const bool bScatterAsync = bIsMultithreaded && GIoDispatcherForceSynchronousScatter == 0 &&
+			(!BlockToDecompress->CompressionMethod.IsNone() ||
+			 BlockToDecompress->EncryptionKey.IsValid() ||
+			 BlockToDecompress->SignatureHash);
 		if (bScatterAsync)
 		{
-			TGraphTask<FDecompressAsyncTask>::CreateTask().ConstructAndDispatchWhenReady(*this, BlockToDecompress);
+			UE::Tasks::Launch(
+				TEXT("ScatterBlockDecompressionTask"),
+				[this, BlockToDecompress]
+				{
+					ScatterBlock(BlockToDecompress, true);
+				},
+				(CPrio_IoDispatcherTaskPriority.Get() == ENamedThreads::BackgroundThreadPriority) ?
+				UE::Tasks::ETaskPriority::BackgroundNormal :
+				UE::Tasks::ETaskPriority::Normal
+			);
 		}
 		else
 		{
diff --git a/Engine/Source/Runtime/PakFile/Private/IoDispatcherFileBackend.h b/Engine/Source/Runtime/PakFile/Private/IoDispatcherFileBackend.h
--- a/Engine/Source/Runtime/PakFile/Private/IoDispatcherFileBackend.h
+++ b/Engine/Source/Runtime/PakFile/Private/IoDispatcherFileBackend.h
@@ -8,7 +8,6 @@
 #include "Containers/Array.h"
 #include "Containers/Map.h"
 #include "Stats/Stats.h"
-#include "Async/TaskGraphInterfaces.h"
 #include "HAL/Runnable.h"
 #include "Misc/AES.h"
 #include "GenericPlatform/GenericPlatformFile.h"
@@ -125,37 +124,6 @@ public:
 	uint32 GetThreadId() const;
 
 private:
-	class FDecompressAsyncTask
-	{
-	public:
-		FDecompressAsyncTask(FFileIoStore& InOuter, FFileIoStoreCompressedBlock* InCompressedBlock)
-			: Outer(InOuter)
-			, CompressedBlock(InCompressedBlock)
-		{
-
-		}
-
-		static FORCEINLINE TStatId GetStatId()
-		{
-			RETURN_QUICK_DECLARE_CYCLE_STAT(FIoStoreDecompressTask, STATGROUP_TaskGraphTasks);
-		}
-
-		static ENamedThreads::Type GetDesiredThread();
-
-		FORCEINLINE static ESubsequentsMode::Type GetSubsequentsMode()
-		{
-			return ESubsequentsMode::FireAndForget;
-		}
-
-		void DoTask(ENamedThreads::Type CurrentThread, const FGraphEventRef& MyCompletionGraphEvent)
-		{
-			Outer.ScatterBlock(CompressedBlock, true);
-		}
-
-	private:
-		FFileIoStore& Outer;
-		FFileIoStoreCompressedBlock* CompressedBlock;
-	};
 
 	bool Resolve(FIoRequestImpl* Request);
 	void OnNewPendingRequestsAdded();
