[IoDispatcher]
 - Fix potential deadlock on low core count when all workers are waiting on IO and the IoDispatcher is waiting on decompression tasks

[RN] Fix potential deadlock on low core count

diff --git a/Engine/Source/Runtime/PakFile/Private/IoDispatcherFileBackend.cpp b/Engine/Source/Runtime/PakFile/Private/IoDispatcherFileBackend.cpp
--- a/Engine/Source/Runtime/PakFile/Private/IoDispatcherFileBackend.cpp
+++ b/Engine/Source/Runtime/PakFile/Private/IoDispatcherFileBackend.cpp
@@ -98,6 +98,13 @@ static FAutoConsoleVariableRef CVar_IoDispatcherTocsEnablePerfectHashing(
 	TEXT("Enable perfect hashmap lookups for iostore tocs")
 );
 
+int32 GIoDispatcherCanDecompressOnStarvation = 1;
+static FAutoConsoleVariableRef CVar_IoDispatcherCanDecompressOnStarvation(
+	TEXT("s.IoDispatcherCanDecompressOnStarvation"),
+	GIoDispatcherCanDecompressOnStarvation,
+	TEXT("IoDispatcher thread will help with decompression tasks when all worker threads are IO starved to avoid deadlocks on low core count")
+);
+
 int32 GIoDispatcherForceSynchronousScatter = 0;
 static FAutoConsoleVariableRef CVar_IoDispatcherForceSynchronousScatter(
 	TEXT("s.IoDispatcherForceSynchronousScatter"),
@@ -1269,6 +1276,15 @@ void FFileIoStore::Initialize(TSharedRef<const FIoDispatcherBackendContext> InCo
 	}
 
 	Thread = FRunnableThread::Create(this, TEXT("IoService"), 0, TPri_AboveNormal);
+
+	using namespace LowLevelTasks;
+	OversubscriptionLimitReached = 
+		FScheduler::Get().GetOversubscriptionLimitReachedEvent().AddLambda(
+			[this]()
+			{
+				BackendContext->WakeUpDispatcherThreadDelegate.Execute();
+			}
+		);
 }
 
 void FFileIoStore::StopThread()
@@ -1283,6 +1299,9 @@ void FFileIoStore::StopThread()
 void FFileIoStore::Shutdown()
 {
 	StopThread();
+
+	using namespace LowLevelTasks;
+	FScheduler::Get().GetOversubscriptionLimitReachedEvent().Remove(OversubscriptionLimitReached);
 }
 
 TIoStatusOr<FIoContainerHeader> FFileIoStore::Mount(const TCHAR* InTocPath, int32 Order, const FGuid& EncryptionKeyGuid, const FAES::FAESKey& EncryptionKey)
@@ -1612,8 +1631,6 @@ void FFileIoStore::ScatterBlock(FFileIoStoreCompressedBlock* CompressedBlock, bo
 		FScopeLock Lock(&DecompressedBlocksCritical);
 		CompressedBlock->Next = FirstDecompressedBlock;
 		FirstDecompressedBlock = CompressedBlock;
-
-		BackendContext->WakeUpDispatcherThreadDelegate.Execute();
 	}
 }
 
@@ -1689,6 +1706,21 @@ void FFileIoStore::FinalizeCompressedBlock(FFileIoStoreCompressedBlock* Compress
 	}
 }
 
+namespace FileIoStoreImpl
+{
+	static std::atomic<int32> ActiveScatterTasks{ 0 };
+
+	bool HasActiveScatterTasks()
+	{
+		return ActiveScatterTasks.load(std::memory_order_relaxed) > 0;
+	}
+
+	bool IsSchedulerOversubscribed(UE::Tasks::ETaskPriority TaskPriority)
+	{
+		return LowLevelTasks::FScheduler::Get().IsOversubscriptionLimitReached(TaskPriority);
+	}
+}
+
 FIoRequestImpl* FFileIoStore::GetCompletedIoRequests()
 {
 	LLM_SCOPE(ELLMTag::FileSystem);
@@ -1815,6 +1847,25 @@ FIoRequestImpl* FFileIoStore::GetCompletedIoRequests()
 		BlockToReap = Next;
 	}
 
+	// Cleanup finished decompression tasks
+	UE::Tasks::FTask* DecompressionTask;
+	while ((DecompressionTask = DecompressionTasks.Peek()) != nullptr && DecompressionTask->IsCompleted())
+	{
+		DecompressionTasks.Dequeue();
+	}
+
+	// Help with decompression to avoid deadlock on low-core count when all tasks threads are busy or waiting on IO requests.
+	if (GIoDispatcherCanDecompressOnStarvation && !FileIoStoreImpl::HasActiveScatterTasks() && (DecompressionTask = DecompressionTasks.Peek()) != nullptr)
+	{
+		if (FileIoStoreImpl::IsSchedulerOversubscribed(DecompressionTask->GetPriority()))
+		{
+			// Try to execute it on the current thread if not already started, no-op if already started.
+			DecompressionTask->TryRetractAndExecute();
+			// In both case we can get rid of it right away since we know progress is being made.
+			DecompressionTasks.Dequeue();
+		}
+	}
+
 	FFileIoStoreCompressedBlock* BlockToDecompress = ReadyForDecompressionHead;
 	while (BlockToDecompress)
 	{
@@ -1845,22 +1896,37 @@ FIoRequestImpl* FFileIoStore::GetCompletedIoRequests()
 			}
 		}
 
+		const UE::Tasks::ETaskPriority IoDispatcherTaskPriority =
+			CPrio_IoDispatcherTaskPriority.Get() == ENamedThreads::BackgroundThreadPriority ?
+				UE::Tasks::ETaskPriority::BackgroundNormal :
+				UE::Tasks::ETaskPriority::Normal;
+
 		// Scatter block asynchronous when the block is compressed, encrypted or signed
 		const bool bScatterAsync = bIsMultithreaded && GIoDispatcherForceSynchronousScatter == 0 &&
 			(!BlockToDecompress->CompressionMethod.IsNone() ||
 			 BlockToDecompress->EncryptionKey.IsValid() ||
-			 BlockToDecompress->SignatureHash);
+			 BlockToDecompress->SignatureHash) && 
+			 // If we're already oversubscribed, we might not receive any further event to wake us and 
+			 // allow us to process our queue. In that case we simply run decompression locally.
+			 !FileIoStoreImpl::IsSchedulerOversubscribed(IoDispatcherTaskPriority);
+
 		if (bScatterAsync)
 		{
-			UE::Tasks::Launch(
-				TEXT("ScatterBlockDecompressionTask"),
-				[this, BlockToDecompress]
-				{
-					ScatterBlock(BlockToDecompress, true);
-				},
-				(CPrio_IoDispatcherTaskPriority.Get() == ENamedThreads::BackgroundThreadPriority) ?
-				UE::Tasks::ETaskPriority::BackgroundNormal :
-				UE::Tasks::ETaskPriority::Normal
+			DecompressionTasks.Enqueue(
+				UE::Tasks::Launch(
+					TEXT("ScatterBlockDecompressionTask"),
+					[this, BlockToDecompress]
+					{
+						FileIoStoreImpl::ActiveScatterTasks++;
+						ScatterBlock(BlockToDecompress, true);
+						FileIoStoreImpl::ActiveScatterTasks--;
+
+						// Important that the notification goes after the decrement of the active scatter tasks
+						// otherwise we could end up missing an event and deadlock.
+						BackendContext->WakeUpDispatcherThreadDelegate.Execute();
+					},
+					IoDispatcherTaskPriority
+				)
 			);
 		}
 		else
diff --git a/Engine/Source/Runtime/PakFile/Private/IoDispatcherFileBackend.h b/Engine/Source/Runtime/PakFile/Private/IoDispatcherFileBackend.h
--- a/Engine/Source/Runtime/PakFile/Private/IoDispatcherFileBackend.h
+++ b/Engine/Source/Runtime/PakFile/Private/IoDispatcherFileBackend.h
@@ -7,7 +7,9 @@
 #include "IO/IoStore.h"
 #include "Containers/Array.h"
 #include "Containers/Map.h"
+#include "Containers/SpscQueue.h"
 #include "Stats/Stats.h"
+#include "Tasks/Task.h"
 #include "HAL/Runnable.h"
 #include "Misc/AES.h"
 #include "GenericPlatform/GenericPlatformFile.h"
@@ -151,6 +153,7 @@ private:
 	mutable FRWLock IoStoreReadersLock;
 	TArray<TUniquePtr<FFileIoStoreReader>> IoStoreReaders;
 	TArray<TUniquePtr<FFileIoStoreCompressionContext>> CompressionContexts;
+	TSpscQueue<UE::Tasks::FTask> DecompressionTasks;
 	FFileIoStoreCompressionContext* FirstFreeCompressionContext = nullptr;
 	FFileIoStoreCompressedBlock* ReadyForDecompressionHead = nullptr;
 	FFileIoStoreCompressedBlock* ReadyForDecompressionTail = nullptr;
@@ -158,6 +161,7 @@ private:
 	FFileIoStoreCompressedBlock* FirstDecompressedBlock = nullptr;
 	FIoRequestImpl* CompletedRequestsHead = nullptr;
 	FIoRequestImpl* CompletedRequestsTail = nullptr;
+	FDelegateHandle OversubscriptionLimitReached;
 };
 
 TSharedRef<FFileIoStore> CreateIoDispatcherFileBackend();
