Process cancelled file read requests before performing file I/O to reduce the risk of stalling the I/O dispatcher thread and GC hang.

- disabled by default, enable with CVAR s.IoDispatcherCullCancelledReadRequests = True

#rnx

diff --git a/Engine/Source/Runtime/PakFile/Internal/IoDispatcherFileBackendTypes.h b/Engine/Source/Runtime/PakFile/Internal/IoDispatcherFileBackendTypes.h
--- a/Engine/Source/Runtime/PakFile/Internal/IoDispatcherFileBackendTypes.h
+++ b/Engine/Source/Runtime/PakFile/Internal/IoDispatcherFileBackendTypes.h
@@ -639,6 +639,7 @@ public:
 	PAKFILE_API TArray<FFileIoStoreReadRequest*> StealRequests();
 	// Remove all requests whose priority has been changed to something other than the Priority of this queue
 	PAKFILE_API TArray<FFileIoStoreReadRequest*> RemoveMisprioritizedRequests();
+	PAKFILE_API void RemoveCancelledRequests(TArray<FFileIoStoreReadRequest*>& OutCancelled);
 
 	PAKFILE_API FFileIoStoreReadRequest* Pop(FFileIoStoreReadRequestSortKey LastSortKey);
 	PAKFILE_API void Push(FFileIoStoreReadRequest* Request);
@@ -665,6 +666,7 @@ class FFileIoStoreRequestQueue
 {
 public:
 	PAKFILE_API FFileIoStoreReadRequest* Pop();
+	PAKFILE_API void PopCancelled(TArray<FFileIoStoreReadRequest*>& OutCancelled);
 	PAKFILE_API void Push(FFileIoStoreReadRequest& Request);	// Takes ownership of Request and rewrites its intrustive linked list pointers
 	PAKFILE_API void Push(FFileIoStoreReadRequestList& Requests); // Consumes the request list and overwrites all intrustive linked list pointers
 	PAKFILE_API void UpdateOrder();
diff --git a/Engine/Source/Runtime/PakFile/Private/GenericPlatformIoDispatcher.cpp b/Engine/Source/Runtime/PakFile/Private/GenericPlatformIoDispatcher.cpp
--- a/Engine/Source/Runtime/PakFile/Private/GenericPlatformIoDispatcher.cpp
+++ b/Engine/Source/Runtime/PakFile/Private/GenericPlatformIoDispatcher.cpp
@@ -11,6 +11,13 @@
 
 //PRAGMA_DISABLE_OPTIMIZATION
 
+bool GIoDispatcherCullCancelledReadRequests = false;
+static FAutoConsoleVariableRef CVar_IoDispatcherCullCancelledReadRequests(
+	TEXT("s.IoDispatcherCullCancelledReadRequests"),
+	GIoDispatcherCullCancelledReadRequests,
+	TEXT("Process cancelled file read requests before starting I/O.")
+	TEXT("This can prevent the I/O dispatcher thread to stall due to the request queue getting too big.")
+);
 
 TRACE_DECLARE_INT_COUNTER_EXTERN(IoDispatcherFileBackendSequentialReads);
 TRACE_DECLARE_INT_COUNTER_EXTERN(IoDispatcherFileBackendForwardSeeks);
@@ -75,6 +82,23 @@ void FGenericFileIoStoreImpl::CloseContainer(uint64 ContainerFileHandle)
 
 bool FGenericFileIoStoreImpl::StartRequests(FFileIoStoreRequestQueue& RequestQueue)
 {
+	if (GIoDispatcherCullCancelledReadRequests)
+	{
+		TArray<FFileIoStoreReadRequest*> Cancelled;
+		RequestQueue.PopCancelled(Cancelled);
+		if (!Cancelled.IsEmpty())
+		{
+			{
+				FScopeLock _(&CompletedRequestsCritical);
+				for (FFileIoStoreReadRequest* Request : Cancelled)
+				{
+					CompletedRequests.Add(Request);
+				}
+			}
+			WakeUpDispatcherThreadDelegate->Execute();
+		}
+	}
+
 	if (!AcquiredBuffer)
 	{
 		AcquiredBuffer = BufferAllocator->AllocBuffer();
diff --git a/Engine/Source/Runtime/PakFile/Private/IoDispatcherFileBackend.cpp b/Engine/Source/Runtime/PakFile/Private/IoDispatcherFileBackend.cpp
--- a/Engine/Source/Runtime/PakFile/Private/IoDispatcherFileBackend.cpp
+++ b/Engine/Source/Runtime/PakFile/Private/IoDispatcherFileBackend.cpp
@@ -334,6 +334,21 @@ TArray<FFileIoStoreReadRequest*> FFileIoStoreOffsetSortedRequestQueue::RemoveMis
 	return RequestsToReturn;
 }
 
+void FFileIoStoreOffsetSortedRequestQueue::RemoveCancelledRequests(TArray<FFileIoStoreReadRequest*>& OutCancelled)
+{
+	for (int32 Idx = Requests.Num() - 1; Idx >= 0; --Idx)
+	{
+		FFileIoStoreReadRequest* Request = Requests[Idx];
+		if (Request->bCancelled)
+		{
+			PeekRequestIndex = INDEX_NONE;
+			OutCancelled.Add(Request);
+			RequestsBySequence.Remove(Request);
+			Requests.RemoveAt(Idx, EAllowShrinking::No);
+		}
+	}
+}
+
 FFileIoStoreReadRequest* FFileIoStoreOffsetSortedRequestQueue::GetNextInternal(FFileIoStoreReadRequestSortKey LastSortKey, bool bPop)
 {
 	if (Requests.Num() == 0)
@@ -538,6 +553,48 @@ FFileIoStoreReadRequest* FFileIoStoreRequestQueue::Pop()
 	return Result;
 }
 
+void FFileIoStoreRequestQueue::PopCancelled(TArray<FFileIoStoreReadRequest*>& OutCancelled)
+{
+	TRACE_CPUPROFILER_EVENT_SCOPE(RequestQueuePopCancelled);
+	FScopeLock _(&CriticalSection);
+	UpdateSortRequestsByOffset();
+	
+	if (bSortRequestsByOffset)
+	{
+		for (FFileIoStoreOffsetSortedRequestQueue& PrioQueue : SortedPriorityQueues)
+		{
+			PrioQueue.RemoveCancelledRequests(OutCancelled);
+		}
+
+		// Pop/Peek rely on empty queues being culled
+		SortedPriorityQueues.RemoveAll([](FFileIoStoreOffsetSortedRequestQueue& SubQueue) { return SubQueue.IsEmpty(); });
+	}
+	else
+	{
+		for (int32 Idx = Heap.Num() - 1; Idx >= 0; --Idx)
+		{
+			FFileIoStoreReadRequest* Request = Heap[Idx];
+			if (Request->bCancelled)
+			{
+				OutCancelled.Add(Request);
+				Heap.RemoveAt(Idx, EAllowShrinking::No);
+			}
+		}
+
+		if (!OutCancelled.IsEmpty())
+		{
+			Heap.Heapify(QueueSortFunc);
+		}
+	}
+
+	for (FFileIoStoreReadRequest* Request : OutCancelled)
+	{
+		check(Request->bCancelled);
+		Request->QueueStatus = FFileIoStoreReadRequest::QueueStatus_Started;
+		Request->ContainerFilePartition->StartedReadRequestsCount.fetch_add(1, std::memory_order_release);
+	}
+}
+
 void FFileIoStoreRequestQueue::PushToPriorityQueues(FFileIoStoreReadRequest* Request)
 {
 	int32 QueueIndex = Algo::LowerBoundBy(SortedPriorityQueues, Request->Priority, QueuePriorityProjection, TLess<int32>());
