Correctly determine IO scatter task priority from CVar.

diff --git a/Engine/Source/Runtime/PakFile/Private/IoDispatcherFileBackend.cpp b/Engine/Source/Runtime/PakFile/Private/IoDispatcherFileBackend.cpp
--- a/Engine/Source/Runtime/PakFile/Private/IoDispatcherFileBackend.cpp
+++ b/Engine/Source/Runtime/PakFile/Private/IoDispatcherFileBackend.cpp
@@ -1954,7 +1954,7 @@ FIoRequestImpl* FFileIoStore::GetCompletedIoRequests()
 		}
 
 		const UE::Tasks::ETaskPriority IoDispatcherTaskPriority =
-			CPrio_IoDispatcherTaskPriority.Get() == ENamedThreads::BackgroundThreadPriority ?
+			EnumHasAnyFlags(CPrio_IoDispatcherTaskPriority.Get(), ENamedThreads::BackgroundThreadPriority) ?
 				UE::Tasks::ETaskPriority::BackgroundNormal :
 				UE::Tasks::ETaskPriority::Normal;
 
