Split a large file with components multiplexed together into smaller parts.

#rnx

diff --git a/Engine/Source/Runtime/PakFile/Private/IPlatformFilePak.cpp b/Engine/Source/Runtime/PakFile/Private/IPlatformFilePak.cpp
--- a/Engine/Source/Runtime/PakFile/Private/IPlatformFilePak.cpp
+++ b/Engine/Source/Runtime/PakFile/Private/IPlatformFilePak.cpp
@@ -12,11 +12,9 @@
 #include "Misc/CoreDelegates.h"
 #include "Misc/CoreDelegatesInternal.h"
 #include "Misc/App.h"
-#include "Modules/ModuleManager.h"
 #include "Misc/ConfigCacheIni.h"
 #include "Misc/SecureHash.h"
 #include "HAL/FileManagerGeneric.h"
-#include "HAL/IPlatformFileModule.h"
 #include "SignedArchiveReader.h"
 #include "Misc/AES.h"
 #include "GenericPlatform/GenericPlatformChunkInstall.h"
@@ -27,14 +25,12 @@
 #include "Misc/Base64.h"
 #include "HAL/DiskUtilizationTracker.h"
 #include "Stats/StatsMisc.h"
-#include "HAL/PlatformFileManager.h"
 #include "HAL/ThreadHeartBeat.h"
 #if !(IS_PROGRAM || WITH_EDITOR)
 #include "Misc/ConfigCacheIni.h"
 #endif
 #include "ProfilingDebugging/CsvProfiler.h"
 #include "Misc/EncryptionKeyManager.h"
-#include "Misc/Fnv.h"
 #include "ProfilingDebugging/ScopedTimers.h"
 #include "Async/MappedFileHandle.h"
 #include "IoDispatcherFileBackend.h"
@@ -49,6 +45,9 @@
 #include "String/RemoveFrom.h"
 #include "Algo/AnyOf.h"
 
+#include "PakFile.inl"
+#include "PakIntervalTree.inl"
+
 DEFINE_LOG_CATEGORY(LogPakFile);
 
 DEFINE_STAT(STAT_PakFile_Read);
@@ -59,8 +58,8 @@
 
 
 #if CSV_PROFILER_STATS
-int64 GTotalLoaded = 0;
-int64 GTotalLoadedLastTick = 0;
+static int64 GTotalLoaded = 0;
+static int64 GTotalLoadedLastTick = 0;
 #endif
 
 
@@ -87,7 +86,94 @@
 
 static FString GMountStartupPaksWildCard = TEXT(MOUNT_STARTUP_PAKS_WILDCARD);
 
+bool ShouldCheckPak()
+{
+	static bool bShouldCheckPak = FParse::Param(FCommandLine::Get(), TEXT("checkpak"));
+	return bShouldCheckPak;
+}
 
+static bool CheckIoStoreContainerBlockSignatures(const TCHAR* InContainerPath)
+{
+	TRACE_CPUPROFILER_EVENT_SCOPE(CheckIoStoreContainerBlockSignatures);
+	UE_LOG(LogPakFile, Display, TEXT("Checking container file \"%s\"..."), InContainerPath);
+	double StartTime = FPlatformTime::Seconds();
 
+	FIoStoreTocResource TocResource;
+	FIoStatus Status = FIoStoreTocResource::Read(InContainerPath, EIoStoreTocReadOptions::Default, TocResource);
+	if (!Status.IsOk())
+	{
+		UE_LOG(LogPakFile, Error, TEXT("Failed reading toc file \"%s\"."), InContainerPath);
+		return false;
+	}
+
+	if (TocResource.ChunkBlockSignatures.Num() != TocResource.CompressionBlocks.Num())
+	{
+		UE_LOG(LogPakFile, Error, TEXT("Toc file \"%s\" doesn't contain any chunk block signatures."), InContainerPath);
+		return false;
+	}
+
+	TUniquePtr<FArchive> ContainerFileReader;
+	int32 LastPartitionIndex = -1;
+	IPlatformFile& Ipf = FPlatformFileManager::Get().GetPlatformFile();
+	TArray<uint8> BlockBuffer;
+	BlockBuffer.SetNum(static_cast<int32>(TocResource.Header.CompressionBlockSize));
+	const int32 BlockCount = TocResource.CompressionBlocks.Num();
+	int32 ErrorCount = 0;
+	FString ContainerBasePath = FPaths::ChangeExtension(InContainerPath, TEXT(""));
+	TStringBuilder<256> UcasFilePath;
+	for (int32 BlockIndex = 0; BlockIndex < BlockCount; ++BlockIndex)
+	{
+		const FIoStoreTocCompressedBlockEntry& CompressionBlockEntry = TocResource.CompressionBlocks[BlockIndex];
+		uint64 BlockRawSize = Align(CompressionBlockEntry.GetCompressedSize(), FAES::AESBlockSize);
+		check(BlockRawSize <= TocResource.Header.CompressionBlockSize);
+		const int32 PartitionIndex = int32(CompressionBlockEntry.GetOffset() / TocResource.Header.PartitionSize);
+		const uint64 PartitionRawOffset = CompressionBlockEntry.GetOffset() % TocResource.Header.PartitionSize;
+		if (PartitionIndex != LastPartitionIndex)
+		{
+			UcasFilePath.Reset();
+			UcasFilePath.Append(ContainerBasePath);
+			if (PartitionIndex > 0)
+			{
+				UcasFilePath.Append(FString::Printf(TEXT("_s%d"), PartitionIndex));
+			}
+			UcasFilePath.Append(TEXT(".ucas"));
+			IFileHandle* ContainerFileHandle = Ipf.OpenRead(*UcasFilePath, /* allowwrite */ false);
+			if (!ContainerFileHandle)
+			{
+				UE_LOG(LogPakFile, Error, TEXT("Failed opening container file \"%s\"."), *UcasFilePath);
+				return false;
+			}
+			ContainerFileReader.Reset(new FArchiveFileReaderGeneric(ContainerFileHandle, *UcasFilePath, ContainerFileHandle->Size(), 256 << 10));
+			LastPartitionIndex = PartitionIndex;
+		}
+		ContainerFileReader->Seek(PartitionRawOffset);
+		ContainerFileReader->Precache(PartitionRawOffset, 0); // Without this buffering won't work due to the first read after a seek always being uncached
+		ContainerFileReader->Serialize(BlockBuffer.GetData(), BlockRawSize);
+		FSHAHash BlockHash;
+		FSHA1::HashBuffer(BlockBuffer.GetData(), BlockRawSize, BlockHash.Hash);
+		if (TocResource.ChunkBlockSignatures[BlockIndex] != BlockHash)
+		{
+			UE_LOG(LogPakFile, Warning, TEXT("Hash mismatch for block [%i/%i]! Expected %s, Received %s"), BlockIndex, BlockCount, *TocResource.ChunkBlockSignatures[BlockIndex].ToString(), *BlockHash.ToString());
+
+			FPakChunkSignatureCheckFailedData Data(*UcasFilePath, TPakChunkHash(), TPakChunkHash(), BlockIndex);
+#if PAKHASH_USE_CRC
+			Data.ExpectedHash = GetTypeHash(TocResource.ChunkBlockSignatures[BlockIndex]);
+			Data.ReceivedHash = GetTypeHash(BlockHash);
+#else
+			Data.ExpectedHash = TocResource.ChunkBlockSignatures[BlockIndex];
+			Data.ReceivedHash = BlockHash;
+#endif
+			FPakPlatformFile::BroadcastPakChunkSignatureCheckFailure(Data);
+			++ErrorCount;
+		}
+	}
+
+	double EndTime = FPlatformTime::Seconds();
+	double ElapsedTime = EndTime - StartTime;
+	UE_LOG(LogPakFile, Display, TEXT("Container file \"%s\" checked in %.2fs"), InContainerPath, ElapsedTime);
+
+	return ErrorCount == 0;
+}
+
 int32 GetPakchunkIndexFromPakFile(const FString& InFilename)
 {
@@ -1121,647 +1207,6 @@
 
 class FPakPrecacher;
 
-typedef uint64 FJoinedOffsetAndPakIndex;
-static FORCEINLINE uint16 GetRequestPakIndexLow(FJoinedOffsetAndPakIndex Joined)
-{
-	return uint16((Joined >> 48) & 0xffff);
-}
-
-static FORCEINLINE int64 GetRequestOffset(FJoinedOffsetAndPakIndex Joined)
-{
-	return int64(Joined & 0xffffffffffffll);
-}
-
-static FORCEINLINE FJoinedOffsetAndPakIndex MakeJoinedRequest(uint16 PakIndex, int64 Offset)
-{
-	check(Offset >= 0);
-	return (FJoinedOffsetAndPakIndex(PakIndex) << 48) | Offset;
-}
-
-enum
-{
-	IntervalTreeInvalidIndex = 0
-};
-
-
-typedef uint32 TIntervalTreeIndex; // this is the arg type of TSparseArray::operator[]
-
-static uint32 GNextSalt = 1;
-
-// This is like TSparseArray, only a bit safer and I needed some restrictions on resizing.
-template<class TItem>
-class TIntervalTreeAllocator
-{
-	TArray<TItem> Items;
-	TArray<int32> FreeItems; //@todo make this into a linked list through the existing items
-	uint32 Salt;
-	uint32 SaltMask;
-public:
-	TIntervalTreeAllocator()
-	{
-		check(GNextSalt < 4);
-		Salt = (GNextSalt++) << 30;
-		SaltMask = MAX_uint32 << 30;
-		verify((Alloc() & ~SaltMask) == IntervalTreeInvalidIndex); // we want this to always have element zero so we can figure out an index from a pointer
-	}
-	inline TIntervalTreeIndex Alloc()
-	{
-		int32 Result;
-		if (FreeItems.Num())
-		{
-			Result = FreeItems.Pop();
-		}
-		else
-		{
-			Result = Items.Num();
-			Items.AddUninitialized();
-
-		}
-		new ((void*)&Items[Result]) TItem();
-		return Result | Salt;;
-	}
-	void EnsureNoRealloc(int32 NeededNewNum)
-	{
-		if (FreeItems.Num() + Items.GetSlack() < NeededNewNum)
-		{
-			Items.Reserve(Items.Num() + NeededNewNum);
-		}
-	}
-	FORCEINLINE TItem& Get(TIntervalTreeIndex InIndex)
-	{
-		TIntervalTreeIndex Index = InIndex & ~SaltMask;
-		check((InIndex & SaltMask) == Salt && Index != IntervalTreeInvalidIndex && Index >= 0 && Index < (uint32)Items.Num()); //&& !FreeItems.Contains(Index));
-		return Items[Index];
-	}
-	FORCEINLINE void Free(TIntervalTreeIndex InIndex)
-	{
-		TIntervalTreeIndex Index = InIndex & ~SaltMask;
-		check((InIndex & SaltMask) == Salt && Index != IntervalTreeInvalidIndex && Index >= 0 && Index < (uint32)Items.Num()); //&& !FreeItems.Contains(Index));
-		Items[Index].~TItem();
-		FreeItems.Push(Index);
-		if (FreeItems.Num() + 1 == Items.Num())
-		{
-			// get rid everything to restore memory coherence
-			Items.Empty();
-			FreeItems.Empty();
-			verify((Alloc() & ~SaltMask) == IntervalTreeInvalidIndex); // we want this to always have element zero so we can figure out an index from a pointer
-		}
-	}
-	FORCEINLINE void CheckIndex(TIntervalTreeIndex InIndex)
-	{
-		TIntervalTreeIndex Index = InIndex & ~SaltMask;
-		check((InIndex & SaltMask) == Salt && Index != IntervalTreeInvalidIndex && Index >= 0 && Index < (uint32)Items.Num()); // && !FreeItems.Contains(Index));
-	}
-};
-
-class FIntervalTreeNode
-{
-public:
-	TIntervalTreeIndex LeftChildOrRootOfLeftList;
-	TIntervalTreeIndex RootOfOnList;
-	TIntervalTreeIndex RightChildOrRootOfRightList;
-
-	FIntervalTreeNode()
-		: LeftChildOrRootOfLeftList(IntervalTreeInvalidIndex)
-		, RootOfOnList(IntervalTreeInvalidIndex)
-		, RightChildOrRootOfRightList(IntervalTreeInvalidIndex)
-	{
-	}
-	~FIntervalTreeNode()
-	{
-		check(LeftChildOrRootOfLeftList == IntervalTreeInvalidIndex && RootOfOnList == IntervalTreeInvalidIndex && RightChildOrRootOfRightList == IntervalTreeInvalidIndex); // this routine does not handle recursive destruction
-	}
-};
-
-static TIntervalTreeAllocator<FIntervalTreeNode> GIntervalTreeNodeNodeAllocator;
-
-static FORCEINLINE uint64 HighBit(uint64 x)
-{
-	return x & (1ull << 63);
-}
-
-static FORCEINLINE bool IntervalsIntersect(uint64 Min1, uint64 Max1, uint64 Min2, uint64 Max2)
-{
-	return !(Max2 < Min1 || Max1 < Min2);
-}
-
-template<typename TItem>
-// this routine assume that the pointers remain valid even though we are reallocating
-static void AddToIntervalTree_Dangerous(
-	TIntervalTreeIndex* RootNode,
-	TIntervalTreeAllocator<TItem>& Allocator,
-	TIntervalTreeIndex Index,
-	uint64 MinInterval,
-	uint64 MaxInterval,
-	uint32 CurrentShift,
-	uint32 MaxShift
-)
-{
-	while (true)
-	{
-		if (*RootNode == IntervalTreeInvalidIndex)
-		{
-			*RootNode = GIntervalTreeNodeNodeAllocator.Alloc();
-		}
-
-		int64 MinShifted = HighBit(MinInterval << CurrentShift);
-		int64 MaxShifted = HighBit(MaxInterval << CurrentShift);
-		FIntervalTreeNode& Root = GIntervalTreeNodeNodeAllocator.Get(*RootNode);
-
-		if (MinShifted == MaxShifted && CurrentShift < MaxShift)
-		{
-			CurrentShift++;
-			RootNode = (!MinShifted) ? &Root.LeftChildOrRootOfLeftList : &Root.RightChildOrRootOfRightList;
-		}
-		else
-		{
-			TItem& Item = Allocator.Get(Index);
-			if (MinShifted != MaxShifted) // crosses middle
-			{
-				Item.Next = Root.RootOfOnList;
-				Root.RootOfOnList = Index;
-			}
-			else // we are at the leaf
-			{
-				if (!MinShifted)
-				{
-					Item.Next = Root.LeftChildOrRootOfLeftList;
-					Root.LeftChildOrRootOfLeftList = Index;
-				}
-				else
-				{
-					Item.Next = Root.RightChildOrRootOfRightList;
-					Root.RightChildOrRootOfRightList = Index;
-				}
-			}
-			return;
-		}
-	}
-}
-
-template<typename TItem>
-static void AddToIntervalTree(
-	TIntervalTreeIndex* RootNode,
-	TIntervalTreeAllocator<TItem>& Allocator,
-	TIntervalTreeIndex Index,
-	uint32 StartShift,
-	uint32 MaxShift
-)
-{
-	GIntervalTreeNodeNodeAllocator.EnsureNoRealloc(1 + MaxShift - StartShift);
-	TItem& Item = Allocator.Get(Index);
-	check(Item.Next == IntervalTreeInvalidIndex);
-	uint64 MinInterval = GetRequestOffset(Item.OffsetAndPakIndex);
-	uint64 MaxInterval = MinInterval + Item.Size - 1;
-	AddToIntervalTree_Dangerous(RootNode, Allocator, Index, MinInterval, MaxInterval, StartShift, MaxShift);
-
-}
-
-template<typename TItem>
-static FORCEINLINE bool ScanNodeListForRemoval(
-	TIntervalTreeIndex* Iter,
-	TIntervalTreeAllocator<TItem>& Allocator,
-	TIntervalTreeIndex Index,
-	uint64 MinInterval,
-	uint64 MaxInterval
-)
-{
-	while (*Iter != IntervalTreeInvalidIndex)
-	{
-
-		TItem& Item = Allocator.Get(*Iter);
-		if (*Iter == Index)
-		{
-			*Iter = Item.Next;
-			Item.Next = IntervalTreeInvalidIndex;
-			return true;
-		}
-		Iter = &Item.Next;
-	}
-	return false;
-}
-
-template<typename TItem>
-static bool RemoveFromIntervalTree(
-	TIntervalTreeIndex* RootNode,
-	TIntervalTreeAllocator<TItem>& Allocator,
-	TIntervalTreeIndex Index,
-	uint64 MinInterval,
-	uint64 MaxInterval,
-	uint32 CurrentShift,
-	uint32 MaxShift
-)
-{
-	bool bResult = false;
-	if (*RootNode != IntervalTreeInvalidIndex)
-	{
-		int64 MinShifted = HighBit(MinInterval << CurrentShift);
-		int64 MaxShifted = HighBit(MaxInterval << CurrentShift);
-		FIntervalTreeNode& Root = GIntervalTreeNodeNodeAllocator.Get(*RootNode);
-
-		if (!MinShifted && !MaxShifted)
-		{
-			if (CurrentShift == MaxShift)
-			{
-				bResult = ScanNodeListForRemoval(&Root.LeftChildOrRootOfLeftList, Allocator, Index, MinInterval, MaxInterval);
-			}
-			else
-			{
-				bResult = RemoveFromIntervalTree(&Root.LeftChildOrRootOfLeftList, Allocator, Index, MinInterval, MaxInterval, CurrentShift + 1, MaxShift);
-			}
-		}
-		else if (!MinShifted && MaxShifted)
-		{
-			bResult = ScanNodeListForRemoval(&Root.RootOfOnList, Allocator, Index, MinInterval, MaxInterval);
-		}
-		else
-		{
-			if (CurrentShift == MaxShift)
-			{
-				bResult = ScanNodeListForRemoval(&Root.RightChildOrRootOfRightList, Allocator, Index, MinInterval, MaxInterval);
-			}
-			else
-			{
-				bResult = RemoveFromIntervalTree(&Root.RightChildOrRootOfRightList, Allocator, Index, MinInterval, MaxInterval, CurrentShift + 1, MaxShift);
-			}
-		}
-		if (bResult)
-		{
-			if (Root.LeftChildOrRootOfLeftList == IntervalTreeInvalidIndex && Root.RootOfOnList == IntervalTreeInvalidIndex && Root.RightChildOrRootOfRightList == IntervalTreeInvalidIndex)
-			{
-				check(&Root == &GIntervalTreeNodeNodeAllocator.Get(*RootNode));
-				GIntervalTreeNodeNodeAllocator.Free(*RootNode);
-				*RootNode = IntervalTreeInvalidIndex;
-			}
-		}
-	}
-	return bResult;
-}
-
-template<typename TItem>
-static bool RemoveFromIntervalTree(
-	TIntervalTreeIndex* RootNode,
-	TIntervalTreeAllocator<TItem>& Allocator,
-	TIntervalTreeIndex Index,
-	uint32 StartShift,
-	uint32 MaxShift
-)
-{
-	TItem& Item = Allocator.Get(Index);
-	uint64 MinInterval = GetRequestOffset(Item.OffsetAndPakIndex);
-	uint64 MaxInterval = MinInterval + Item.Size - 1;
-	return RemoveFromIntervalTree(RootNode, Allocator, Index, MinInterval, MaxInterval, StartShift, MaxShift);
-}
-
-template<typename TItem>
-static FORCEINLINE void ScanNodeListForRemovalFunc(
-	TIntervalTreeIndex* Iter,
-	TIntervalTreeAllocator<TItem>& Allocator,
-	uint64 MinInterval,
-	uint64 MaxInterval,
-	TFunctionRef<bool(TIntervalTreeIndex)> Func
-)
-{
-	while (*Iter != IntervalTreeInvalidIndex)
-	{
-		TItem& Item = Allocator.Get(*Iter);
-		uint64 Offset = uint64(GetRequestOffset(Item.OffsetAndPakIndex));
-		uint64 LastByte = Offset + uint64(Item.Size) - 1;
-
-		// save the value and then clear it.
-		TIntervalTreeIndex NextIndex = Item.Next;
-		if (IntervalsIntersect(MinInterval, MaxInterval, Offset, LastByte) && Func(*Iter))
-		{
-			*Iter = NextIndex; // this may have already be deleted, so cannot rely on the memory block
-		}
-		else
-		{
-			Iter = &Item.Next;
-		}
-	}
-}
-
-template<typename TItem>
-static void MaybeRemoveOverlappingNodesInIntervalTree(
-	TIntervalTreeIndex* RootNode,
-	TIntervalTreeAllocator<TItem>& Allocator,
-	uint64 MinInterval,
-	uint64 MaxInterval,
-	uint64 MinNode,
-	uint64 MaxNode,
-	uint32 CurrentShift,
-	uint32 MaxShift,
-	TFunctionRef<bool(TIntervalTreeIndex)> Func
-)
-{
-	if (*RootNode != IntervalTreeInvalidIndex)
-	{
-		int64 MinShifted = HighBit(MinInterval << CurrentShift);
-		int64 MaxShifted = HighBit(MaxInterval << CurrentShift);
-		FIntervalTreeNode& Root = GIntervalTreeNodeNodeAllocator.Get(*RootNode);
-		uint64 Center = (MinNode + MaxNode + 1) >> 1;
-
-		//UE_LOG(LogTemp, Warning, TEXT("Exploring Node %X [%d, %d] %d%d     interval %llX %llX    node interval %llX %llX   center %llX  "), *RootNode, CurrentShift, MaxShift, !!MinShifted, !!MaxShifted, MinInterval, MaxInterval, MinNode, MaxNode, Center);
-
-
-		if (!MinShifted)
-		{
-			if (CurrentShift == MaxShift)
-			{
-				//UE_LOG(LogTemp, Warning, TEXT("LeftBottom %X [%d, %d] %d%d"), *RootNode, CurrentShift, MaxShift, !!MinShifted, !!MaxShifted);
-				ScanNodeListForRemovalFunc(&Root.LeftChildOrRootOfLeftList, Allocator, MinInterval, MaxInterval, Func);
-			}
-			else
-			{
-				//UE_LOG(LogTemp, Warning, TEXT("LeftRecur %X [%d, %d] %d%d"), *RootNode, CurrentShift, MaxShift, !!MinShifted, !!MaxShifted);
-				MaybeRemoveOverlappingNodesInIntervalTree(&Root.LeftChildOrRootOfLeftList, Allocator, MinInterval, FMath::Min(MaxInterval, Center - 1), MinNode, Center - 1, CurrentShift + 1, MaxShift, Func);
-			}
-		}
-
-		//UE_LOG(LogTemp, Warning, TEXT("Center %X [%d, %d] %d%d"), *RootNode, CurrentShift, MaxShift, !!MinShifted, !!MaxShifted);
-		ScanNodeListForRemovalFunc(&Root.RootOfOnList, Allocator, MinInterval, MaxInterval, Func);
-
-		if (MaxShifted)
-		{
-			if (CurrentShift == MaxShift)
-			{
-				//UE_LOG(LogTemp, Warning, TEXT("RightBottom %X [%d, %d] %d%d"), *RootNode, CurrentShift, MaxShift, !!MinShifted, !!MaxShifted);
-				ScanNodeListForRemovalFunc(&Root.RightChildOrRootOfRightList, Allocator, MinInterval, MaxInterval, Func);
-			}
-			else
-			{
-				//UE_LOG(LogTemp, Warning, TEXT("RightRecur %X [%d, %d] %d%d"), *RootNode, CurrentShift, MaxShift, !!MinShifted, !!MaxShifted);
-				MaybeRemoveOverlappingNodesInIntervalTree(&Root.RightChildOrRootOfRightList, Allocator, FMath::Max(MinInterval, Center), MaxInterval, Center, MaxNode, CurrentShift + 1, MaxShift, Func);
-			}
-		}
-
-		//UE_LOG(LogTemp, Warning, TEXT("Done Exploring Node %X [%d, %d] %d%d     interval %llX %llX    node interval %llX %llX   center %llX  "), *RootNode, CurrentShift, MaxShift, !!MinShifted, !!MaxShifted, MinInterval, MaxInterval, MinNode, MaxNode, Center);
-
-		if (Root.LeftChildOrRootOfLeftList == IntervalTreeInvalidIndex && Root.RootOfOnList == IntervalTreeInvalidIndex && Root.RightChildOrRootOfRightList == IntervalTreeInvalidIndex)
-		{
-			check(&Root == &GIntervalTreeNodeNodeAllocator.Get(*RootNode));
-			GIntervalTreeNodeNodeAllocator.Free(*RootNode);
-			*RootNode = IntervalTreeInvalidIndex;
-		}
-	}
-}
-
-
-template<typename TItem>
-static FORCEINLINE bool ScanNodeList(
-	TIntervalTreeIndex Iter,
-	TIntervalTreeAllocator<TItem>& Allocator,
-	uint64 MinInterval,
-	uint64 MaxInterval,
-	TFunctionRef<bool(TIntervalTreeIndex)> Func
-)
-{
-	while (Iter != IntervalTreeInvalidIndex)
-	{
-		TItem& Item = Allocator.Get(Iter);
-		uint64 Offset = uint64(GetRequestOffset(Item.OffsetAndPakIndex));
-		uint64 LastByte = Offset + uint64(Item.Size) - 1;
-		if (IntervalsIntersect(MinInterval, MaxInterval, Offset, LastByte))
-		{
-			if (!Func(Iter))
-			{
-				return false;
-			}
-		}
-		Iter = Item.Next;
-	}
-	return true;
-}
-
-template<typename TItem>
-static bool OverlappingNodesInIntervalTree(
-	TIntervalTreeIndex RootNode,
-	TIntervalTreeAllocator<TItem>& Allocator,
-	uint64 MinInterval,
-	uint64 MaxInterval,
-	uint64 MinNode,
-	uint64 MaxNode,
-	uint32 CurrentShift,
-	uint32 MaxShift,
-	TFunctionRef<bool(TIntervalTreeIndex)> Func
-)
-{
-	if (RootNode != IntervalTreeInvalidIndex)
-	{
-		int64 MinShifted = HighBit(MinInterval << CurrentShift);
-		int64 MaxShifted = HighBit(MaxInterval << CurrentShift);
-		FIntervalTreeNode& Root = GIntervalTreeNodeNodeAllocator.Get(RootNode);
-		uint64 Center = (MinNode + MaxNode + 1) >> 1;
-
-		if (!MinShifted)
-		{
-			if (CurrentShift == MaxShift)
-			{
-				if (!ScanNodeList(Root.LeftChildOrRootOfLeftList, Allocator, MinInterval, MaxInterval, Func))
-				{
-					return false;
-				}
-			}
-			else
-			{
-				if (!OverlappingNodesInIntervalTree(Root.LeftChildOrRootOfLeftList, Allocator, MinInterval, FMath::Min(MaxInterval, Center - 1), MinNode, Center - 1, CurrentShift + 1, MaxShift, Func))
-				{
-					return false;
-				}
-			}
-		}
-		if (!ScanNodeList(Root.RootOfOnList, Allocator, MinInterval, MaxInterval, Func))
-		{
-			return false;
-		}
-		if (MaxShifted)
-		{
-			if (CurrentShift == MaxShift)
-			{
-				if (!ScanNodeList(Root.RightChildOrRootOfRightList, Allocator, MinInterval, MaxInterval, Func))
-				{
-					return false;
-				}
-			}
-			else
-			{
-				if (!OverlappingNodesInIntervalTree(Root.RightChildOrRootOfRightList, Allocator, FMath::Max(MinInterval, Center), MaxInterval, Center, MaxNode, CurrentShift + 1, MaxShift, Func))
-				{
-					return false;
-				}
-			}
-		}
-	}
-	return true;
-}
-
-template<typename TItem>
-static bool ScanNodeListWithShrinkingInterval(
-	TIntervalTreeIndex Iter,
-	TIntervalTreeAllocator<TItem>& Allocator,
-	uint64 MinInterval,
-	uint64& MaxInterval,
-	TFunctionRef<bool(TIntervalTreeIndex)> Func
-)
-{
-	while (Iter != IntervalTreeInvalidIndex)
-	{
-		TItem& Item = Allocator.Get(Iter);
-		uint64 Offset = uint64(GetRequestOffset(Item.OffsetAndPakIndex));
-		uint64 LastByte = Offset + uint64(Item.Size) - 1;
-		//UE_LOG(LogTemp, Warning, TEXT("Test Overlap %llu %llu %llu %llu"), MinInterval, MaxInterval, Offset, LastByte);
-		if (IntervalsIntersect(MinInterval, MaxInterval, Offset, LastByte))
-		{
-			//UE_LOG(LogTemp, Warning, TEXT("Overlap %llu %llu %llu %llu"), MinInterval, MaxInterval, Offset, LastByte);
-			if (!Func(Iter))
-			{
-				return false;
-			}
-		}
-		Iter = Item.Next;
-	}
-	return true;
-}
-
-template<typename TItem>
-static bool OverlappingNodesInIntervalTreeWithShrinkingInterval(
-	TIntervalTreeIndex RootNode,
-	TIntervalTreeAllocator<TItem>& Allocator,
-	uint64 MinInterval,
-	uint64& MaxInterval,
-	uint64 MinNode,
-	uint64 MaxNode,
-	uint32 CurrentShift,
-	uint32 MaxShift,
-	TFunctionRef<bool(TIntervalTreeIndex)> Func
-)
-{
-	if (RootNode != IntervalTreeInvalidIndex)
-	{
-
-		int64 MinShifted = HighBit(MinInterval << CurrentShift);
-		int64 MaxShifted = HighBit(FMath::Min(MaxInterval, MaxNode) << CurrentShift); // since MaxInterval is changing, we cannot clamp it during recursion.
-		FIntervalTreeNode& Root = GIntervalTreeNodeNodeAllocator.Get(RootNode);
-		uint64 Center = (MinNode + MaxNode + 1) >> 1;
-
-		if (!MinShifted)
-		{
-			if (CurrentShift == MaxShift)
-			{
-				if (!ScanNodeListWithShrinkingInterval(Root.LeftChildOrRootOfLeftList, Allocator, MinInterval, MaxInterval, Func))
-				{
-					return false;
-				}
-			}
-			else
-			{
-				if (!OverlappingNodesInIntervalTreeWithShrinkingInterval(Root.LeftChildOrRootOfLeftList, Allocator, MinInterval, MaxInterval, MinNode, Center - 1, CurrentShift + 1, MaxShift, Func)) // since MaxInterval is changing, we cannot clamp it during recursion.
-				{
-					return false;
-				}
-			}
-		}
-		if (!ScanNodeListWithShrinkingInterval(Root.RootOfOnList, Allocator, MinInterval, MaxInterval, Func))
-		{
-			return false;
-		}
-		MaxShifted = HighBit(FMath::Min(MaxInterval, MaxNode) << CurrentShift); // since MaxInterval is changing, we cannot clamp it during recursion.
-		if (MaxShifted)
-		{
-			if (CurrentShift == MaxShift)
-			{
-				if (!ScanNodeListWithShrinkingInterval(Root.RightChildOrRootOfRightList, Allocator, MinInterval, MaxInterval, Func))
-				{
-					return false;
-				}
-			}
-			else
-			{
-				if (!OverlappingNodesInIntervalTreeWithShrinkingInterval(Root.RightChildOrRootOfRightList, Allocator, FMath::Max(MinInterval, Center), MaxInterval, Center, MaxNode, CurrentShift + 1, MaxShift, Func))
-				{
-					return false;
-				}
-			}
-		}
-	}
-	return true;
-}
-
-
-template<typename TItem>
-static void MaskInterval(
-	TIntervalTreeIndex Index,
-	TIntervalTreeAllocator<TItem>& Allocator,
-	uint64 MinInterval,
-	uint64 MaxInterval,
-	uint32 BytesToBitsShift,
-	uint64* Bits
-)
-{
-	TItem& Item = Allocator.Get(Index);
-	uint64 Offset = uint64(GetRequestOffset(Item.OffsetAndPakIndex));
-	uint64 LastByte = Offset + uint64(Item.Size) - 1;
-	uint64 InterMinInterval = FMath::Max(MinInterval, Offset);
-	uint64 InterMaxInterval = FMath::Min(MaxInterval, LastByte);
-	if (InterMinInterval <= InterMaxInterval)
-	{
-		uint32 FirstBit = uint32((InterMinInterval - MinInterval) >> BytesToBitsShift);
-		uint32 LastBit = uint32((InterMaxInterval - MinInterval) >> BytesToBitsShift);
-		uint32 FirstQWord = FirstBit >> 6;
-		uint32 LastQWord = LastBit >> 6;
-		uint32 FirstBitQWord = FirstBit & 63;
-		uint32 LastBitQWord = LastBit & 63;
-		if (FirstQWord == LastQWord)
-		{
-			Bits[FirstQWord] |= ((MAX_uint64 << FirstBitQWord) & (MAX_uint64 >> (63 - LastBitQWord)));
-		}
-		else
-		{
-			Bits[FirstQWord] |= (MAX_uint64 << FirstBitQWord);
-			for (uint32 QWordIndex = FirstQWord + 1; QWordIndex < LastQWord; QWordIndex++)
-			{
-				Bits[QWordIndex] = MAX_uint64;
-			}
-			Bits[LastQWord] |= (MAX_uint64 >> (63 - LastBitQWord));
-		}
-	}
-}
-
-
-
-template<typename TItem>
-static void OverlappingNodesInIntervalTreeMask(
-	TIntervalTreeIndex RootNode,
-	TIntervalTreeAllocator<TItem>& Allocator,
-	uint64 MinInterval,
-	uint64 MaxInterval,
-	uint64 MinNode,
-	uint64 MaxNode,
-	uint32 CurrentShift,
-	uint32 MaxShift,
-	uint32 BytesToBitsShift,
-	uint64* Bits
-)
-{
-	OverlappingNodesInIntervalTree(
-		RootNode,
-		Allocator,
-		MinInterval,
-		MaxInterval,
-		MinNode,
-		MaxNode,
-		CurrentShift,
-		MaxShift,
-		[&Allocator, MinInterval, MaxInterval, BytesToBitsShift, Bits](TIntervalTreeIndex Index) -> bool
-	{
-		MaskInterval(Index, Allocator, MinInterval, MaxInterval, BytesToBitsShift, Bits);
-		return true;
-	}
-	);
-}
-
-
-
 class IPakRequestor
 {
 	friend class FPakPrecacher;
@@ -5823,2132 +5268,6 @@
 }
 
 
-#if IS_PROGRAM
-FPakFile::FPakFile(const TCHAR* Filename, bool bIsSigned)
-	: PakFilename(Filename)
-	, PakFilenameName(Filename)
-	, PathHashSeed(0)
-	, NumEntries(0)
-	, CachedTotalSize(0)
-	, bSigned(bIsSigned)
-	, bIsValid(false)
-	, bHasPathHashIndex(false)
-	, bHasFullDirectoryIndex(false)
-#if ENABLE_PAKFILE_RUNTIME_PRUNING
-	, bWillPruneDirectoryIndex(false)
-	, bNeedsLegacyPruning(false)
-#endif
-	, PakchunkIndex(GetPakchunkIndexFromPakFile(Filename))
- 	, MappedFileHandle(nullptr)
-	, CacheType(FPakFile::ECacheType::Shared)
-	, CacheIndex(-1)
-	, UnderlyingCacheTrimDisabled(false)
-	, bIsMounted(false)
-{
-	FSharedPakReader Reader = GetSharedReader(NULL);
-	if (Reader)
-	{
-		Timestamp = IFileManager::Get().GetTimeStamp(Filename);
-		Initialize(Reader.GetArchive());
-	}
-}
-#endif
-
-FPakFile::FPakFile(IPlatformFile* LowerLevel, const TCHAR* Filename, bool bIsSigned, bool bLoadIndex)
-	: PakFilename(Filename)
-	, PakFilenameName(Filename)
-	, PathHashSeed(0)
-	, NumEntries(0)
-	, CachedTotalSize(0)
-	, bSigned(bIsSigned)
-	, bIsValid(false)
-	, bHasPathHashIndex(false)
-	, bHasFullDirectoryIndex(false)
-#if ENABLE_PAKFILE_RUNTIME_PRUNING
-	, bWillPruneDirectoryIndex(false)
-	, bNeedsLegacyPruning(false)
-#endif
-	, PakchunkIndex(GetPakchunkIndexFromPakFile(Filename))
-	, MappedFileHandle(nullptr)
-	, CacheType(FPakFile::ECacheType::Shared)
-	, CacheIndex(-1)
-	, UnderlyingCacheTrimDisabled(false)
-	, bIsMounted(false)
-{
-	FSharedPakReader Reader = GetSharedReader(LowerLevel);
-	if (Reader)
-	{
-		Timestamp = LowerLevel->GetTimeStamp(Filename);
-		Initialize(Reader.GetArchive(), bLoadIndex);
-	}
-}
-
-#if WITH_EDITOR
-FPakFile::FPakFile(FArchive* Archive)
-	: PathHashSeed(0)
-	, NumEntries(0)
-	, bSigned(false)
-	, bIsValid(false)
-	, bHasPathHashIndex(false)
-	, bHasFullDirectoryIndex(false)
-#if ENABLE_PAKFILE_RUNTIME_PRUNING
-	, bWillPruneDirectoryIndex(false)
-	, bNeedsLegacyPruning(false)
-#endif
-	, PakchunkIndex(INDEX_NONE)
-	, MappedFileHandle(nullptr)
-	, CacheType(FPakFile::ECacheType::Shared)
-	, CacheIndex(-1)
-	, UnderlyingCacheTrimDisabled(false)
-	, bIsMounted(false)
-{
-	Initialize(*Archive);
-}
-#endif
-
-FPakFile::~FPakFile()
-{
-	delete MappedFileHandle;
-}
-
-bool FPakFile::PassedSignatureChecks() const
-{
-	return Decryptor.IsValid() && Decryptor->IsValid();
-}
-
-
-FArchive* FPakFile::CreatePakReader(IPlatformFile* LowerLevel, const TCHAR* Filename)
-{
-	auto MakeArchive = [&]() -> FArchive* { 
-		if (LowerLevel)
-		{
-			if( IFileHandle* Handle = LowerLevel->OpenRead(Filename) )
-			{
-				return new FArchiveFileReaderGeneric(Handle, Filename, Handle->Size());
-			}
-			else
-			{
-				return nullptr;
-			}
-		}
-		else
-		{
-			return IFileManager::Get().CreateFileReader(Filename); 
-		}
-	};
-
-	bool bNeedsDecryptor = false;
-	if (FPlatformProperties::RequiresCookedData())
-	{
-		bool bShouldCheckSignature = bSigned || FParse::Param(FCommandLine::Get(), TEXT("signedpak")) || FParse::Param(FCommandLine::Get(), TEXT("signed"));
-#if !UE_BUILD_SHIPPING
-		bShouldCheckSignature &= !FParse::Param(FCommandLine::Get(), TEXT("FileOpenLog"));
-#endif
-		if (bShouldCheckSignature)
-		{
-			bNeedsDecryptor = true;
-		}			
-	}
-
-	if(bNeedsDecryptor && !Decryptor.IsValid())
-	{
-		TUniquePtr<FArchive> DecryptorReader{ MakeArchive() };
-		if (DecryptorReader.IsValid())
-		{
-			Decryptor = MakeUnique<FChunkCacheWorker>(MoveTemp(DecryptorReader), Filename);
-		}
-
-		if (!Decryptor.IsValid() || !Decryptor->IsValid())
-		{
-			return nullptr;
-		}
-	}
-
-	// Now we either have a Decryptor or we don't need it
-	check(!bNeedsDecryptor || Decryptor.IsValid());
-
-	TUniquePtr<FArchive> Archive{ MakeArchive() };
-	if (!Archive.IsValid())
-	{
-		return nullptr;
-	}
-
-	if (Decryptor.IsValid())
-	{
-		return new FSignedArchiveReader(Archive.Release(), Decryptor.Get());
-	}
-	else
-	{
-		return Archive.Release();
-	}
-}
-
-bool ShouldCheckPak()
-{
-	static bool bShouldCheckPak = FParse::Param(FCommandLine::Get(), TEXT("checkpak"));
-	return bShouldCheckPak;
-}
-
-void FPakFile::Initialize(FArchive& Reader, bool bLoadIndex)
-{
-	CachedTotalSize = Reader.TotalSize();
-	bool bShouldLoad = false;
-	int32 CompatibleVersion = FPakInfo::PakFile_Version_Latest;
-
-	LLM_SCOPE(ELLMTag::FileSystem);
-
-	// Serialize trailer and check if everything is as expected.
-	// start up one to offset the -- below
-	CompatibleVersion++;
-	int64 FileInfoPos = -1;
-	do
-	{
-		// try the next version down
-		CompatibleVersion--;
-
-		FileInfoPos = CachedTotalSize - Info.GetSerializedSize(CompatibleVersion);
-		if (FileInfoPos >= 0)
-		{
-			Reader.Seek(FileInfoPos);
-			Reader.Precache(FileInfoPos, 0); // Inform the archive that we're going to repeatedly serialize from the current location
-
-			SCOPED_BOOT_TIMING("PakFile_SerilizeTrailer");
-
-			// Serialize trailer and check if everything is as expected.
-			Info.Serialize(Reader, CompatibleVersion);
-			if (Info.Magic == FPakInfo::PakFile_Magic)
-			{
-				bShouldLoad = true;
-			}
-		}
-	} while (!bShouldLoad && CompatibleVersion >= FPakInfo::PakFile_Version_Initial);
-
-	if (bShouldLoad)
-	{
-		UE_CLOG(Info.Magic != FPakInfo::PakFile_Magic, LogPakFile, Fatal, TEXT("Trailing magic number (%ud) in '%s' is different than the expected one. Verify your installation."), Info.Magic, *PakFilename);
-		UE_CLOG(!(Info.Version >= FPakInfo::PakFile_Version_Initial && Info.Version <= CompatibleVersion), LogPakFile, Fatal, TEXT("Invalid pak file version (%d) in '%s'. Verify your installation."), Info.Version, *PakFilename);
-		UE_CLOG(!(Info.IndexOffset >= 0 && Info.IndexOffset < CachedTotalSize), LogPakFile, Fatal, TEXT("Index offset for pak file '%s' is invalid (%lld is bigger than file size %lld)"), *PakFilename, Info.IndexOffset, CachedTotalSize);
-		UE_CLOG(!((Info.IndexOffset + Info.IndexSize) >= 0 && (Info.IndexOffset + Info.IndexSize) <= CachedTotalSize), LogPakFile, Fatal, TEXT("Index end offset for pak file '%s' is invalid (%lld)"), *PakFilename, Info.IndexOffset + Info.IndexSize);
-
-		// If we aren't using a dynamic encryption key, process the pak file using the embedded key
-		if (!Info.EncryptionKeyGuid.IsValid() || UE::FEncryptionKeyManager::Get().ContainsKey(Info.EncryptionKeyGuid))
-		{
-			if (bLoadIndex)
-			{
-				LoadIndex(Reader);
-			}
-
-			if (ShouldCheckPak())
-			{
-				ensure(Check());
-			}
-		}
-
-		if (Decryptor.IsValid())
-		{
-			TSharedPtr<const FPakSignatureFile, ESPMode::ThreadSafe> SignatureFile = Decryptor->GetSignatures();
-			if (SignatureFile->SignatureData.Num() == UE_ARRAY_COUNT(FSHAHash::Hash))
-			{
-				bIsValid = (FMemory::Memcmp(SignatureFile->SignatureData.GetData(), Info.IndexHash.Hash, SignatureFile->SignatureData.Num()) == 0);
-			}
-			else
-			{
-				bIsValid = false;
-			}
-		}
-		else
-		{
-			// LoadIndex should crash in case of an error, so just assume everything is ok if we got here.
-			bIsValid = true;
-		}
-	}
-}
-
-void FPakFile::LoadIndex(FArchive& Reader)
-{
-	if (Info.Version >= FPakInfo::PakFile_Version_PathHashIndex)
-	{
-		if (!LoadIndexInternal(Reader))
-		{
-			// Index loading failed. Try again
-			if (!LoadIndexInternal(Reader))
-			{
-				UE_LOG(LogPakFile, Fatal, TEXT("Corrupt pak index detected on pak file: %s"), *PakFilename);
-			}
-		}
-	}
-	else
-	{
-		SCOPED_BOOT_TIMING("PakFile_LoadLegacy");
-		if (!LoadLegacyIndex(Reader))
-		{
-			// Index loading failed. Try again
-			if (!LoadLegacyIndex(Reader))
-			{
-				UE_LOG(LogPakFile, Fatal, TEXT("Corrupt pak index detected on pak file: %s"), *PakFilename);
-			}
-		}
-	}
-}
-
-bool FPakFile::LoadIndexInternal(FArchive& Reader)
-{
-	bHasPathHashIndex = false;
-	bHasFullDirectoryIndex = false;
-#if ENABLE_PAKFILE_RUNTIME_PRUNING
-	bNeedsLegacyPruning = false;
-	bWillPruneDirectoryIndex = false;
-#endif
-
-	FGuardedInt64 IndexEndPosition = FGuardedInt64(Info.IndexOffset) + Info.IndexSize;
-	if (Info.IndexOffset < 0 || 
-		Info.IndexSize < 0 ||
-		IndexEndPosition.InvalidOrGreaterThan(CachedTotalSize) ||
-		IntFitsIn<int32>(Info.IndexSize) == false)
-	{
-		UE_LOG(LogPakFile, Fatal, TEXT("Corrupted index offset in pak file."));
-		return false;
-	}
-
-	TArray<uint8> PrimaryIndexData;
-	Reader.Seek(Info.IndexOffset);
-	PrimaryIndexData.SetNum((int32)(Info.IndexSize));
-	{
-		SCOPED_BOOT_TIMING("PakFile_LoadPrimaryIndex");
-		Reader.Serialize(PrimaryIndexData.GetData(), Info.IndexSize);
-	}
-
-	FSHAHash ComputedHash;
-	{
-		SCOPED_BOOT_TIMING("PakFile_HashPrimaryIndex");
-		if (!DecryptAndValidateIndex(Reader, PrimaryIndexData, Info.IndexHash, ComputedHash))
-		{
-			UE_LOG(LogPakFile, Log, TEXT("Corrupt pak PrimaryIndex detected!"));
-			UE_LOG(LogPakFile, Log, TEXT(" Filename: %s"), *PakFilename);
-			UE_LOG(LogPakFile, Log, TEXT(" Encrypted: %d"), Info.bEncryptedIndex);
-			UE_LOG(LogPakFile, Log, TEXT(" Total Size: %lld"), Reader.TotalSize());
-			UE_LOG(LogPakFile, Log, TEXT(" Index Offset: %lld"), Info.IndexOffset);
-			UE_LOG(LogPakFile, Log, TEXT(" Index Size: %lld"), Info.IndexSize);
-			UE_LOG(LogPakFile, Log, TEXT(" Stored Index Hash: %s"), *Info.IndexHash.ToString());
-			UE_LOG(LogPakFile, Log, TEXT(" Computed Index Hash: %s"), *ComputedHash.ToString());
-			return false;
-		}
-	}
-
-	FMemoryReader PrimaryIndexReader(PrimaryIndexData);
-
-	// Read the scalar data (mount point, numentries, etc) and all entries.
-	NumEntries = 0;
-	PrimaryIndexReader << MountPoint;
-	// We are just deserializing a string, which could get however long. Since we know it's a path
-	// and paths are bound to file system rules, we know it can't get absurdly long (e.g. windows is _at best_ 32k)
-	// we just sanity check it to prevent operating on massive buffers and risking overflows.
-	if (MountPoint.Len() > 65535)
-	{
-		UE_LOG(LogPakFile, Error, TEXT("Corrupt pak index data: MountPoint path is longer than 65k"));
-		return false;
-	}
-
-	MakeDirectoryFromPath(MountPoint);
-	PrimaryIndexReader << NumEntries;
-	if (NumEntries < 0)
-	{
-		UE_LOG(LogPakFile, Error, TEXT("Corrupt pak index data: Negative entries count in pak file."));
-		return false;
-	}
-	PrimaryIndexReader << PathHashSeed;
-
-	bool bReaderHasPathHashIndex = false;
-	int64 PathHashIndexOffset = INDEX_NONE;
-	int64 PathHashIndexSize = 0;
-	FSHAHash PathHashIndexHash;
-	PrimaryIndexReader << bReaderHasPathHashIndex;
-	if (bReaderHasPathHashIndex)
-	{
-		PrimaryIndexReader << PathHashIndexOffset;
-		PrimaryIndexReader << PathHashIndexSize;
-		PrimaryIndexReader << PathHashIndexHash;
-		bReaderHasPathHashIndex = bReaderHasPathHashIndex && PathHashIndexOffset != INDEX_NONE;
-	}
-
-	bool bReaderHasFullDirectoryIndex = false;
-	int64 FullDirectoryIndexOffset = INDEX_NONE;
-	int64 FullDirectoryIndexSize = 0;
-	FSHAHash FullDirectoryIndexHash;
-	PrimaryIndexReader << bReaderHasFullDirectoryIndex;
-	if (bReaderHasFullDirectoryIndex)
-	{
-		PrimaryIndexReader << FullDirectoryIndexOffset;
-		PrimaryIndexReader << FullDirectoryIndexSize;
-		PrimaryIndexReader << FullDirectoryIndexHash;
-		bReaderHasFullDirectoryIndex = bReaderHasFullDirectoryIndex && FullDirectoryIndexOffset  != INDEX_NONE;
-	}
-	{
-		SCOPED_BOOT_TIMING("PakFile_SerializeEncodedEntries");
-		PrimaryIndexReader << EncodedPakEntries;
-	}
-
-	int32 FilesNum = 0;
-	PrimaryIndexReader << FilesNum;
-	if (FilesNum < 0)
-	{
-		// Should not be possible for any values in the PrimaryIndex to be invalid, since we verified the index hash
-		UE_LOG(LogPakFile, Log, TEXT("Corrupt pak PrimaryIndex detected!"));
-		UE_LOG(LogPakFile, Log, TEXT(" FilesNum: %d"), FilesNum);
-		return false;
-	}
-	Files.SetNum(FilesNum);
-	if (FilesNum > 0)
-	{
-		SCOPED_BOOT_TIMING("PakFile_SerializeUnencodedEntries");
-		FPakEntry* FileEntries = Files.GetData();
-		for (int32 FileIndex = 0; FileIndex < FilesNum; ++FileIndex)
-		{
-			FileEntries[FileIndex].Serialize(PrimaryIndexReader, Info.Version);
-		}
-	}
-
-	// Decide which SecondaryIndex(es) to load
-	bool bWillUseFullDirectoryIndex;
-	bool bWillUsePathHashIndex;
-	bool bReadFullDirectoryIndex;
-	if (bReaderHasPathHashIndex && bReaderHasFullDirectoryIndex)
-	{
-		bWillUseFullDirectoryIndex = IsPakKeepFullDirectory();
-		bWillUsePathHashIndex = !bWillUseFullDirectoryIndex;
-#if ENABLE_PAKFILE_RUNTIME_PRUNING
-		bool bWantToReadFullDirectoryIndex = IsPakKeepFullDirectory() || IsPakValidatePruning() || IsPakDelayPruning();
-#else
-		bool bWantToReadFullDirectoryIndex = IsPakKeepFullDirectory();
-#endif
-		bReadFullDirectoryIndex = bReaderHasFullDirectoryIndex && bWantToReadFullDirectoryIndex;
-	}
-	else if (bReaderHasPathHashIndex)
-	{
-		bWillUsePathHashIndex = true;
-		bWillUseFullDirectoryIndex = false;
-		bReadFullDirectoryIndex = false;
-	}
-	else if (bReaderHasFullDirectoryIndex)
-	{
-		// We don't support creating the PathHash Index at runtime; we want to move to having only the PathHashIndex, so supporting not having it at all is not useful enough to write
-		bWillUsePathHashIndex = false;
-		bWillUseFullDirectoryIndex = true;
-		bReadFullDirectoryIndex = true;
-	}
-	else
-	{
-		// It should not be possible for PrimaryIndexes to be built without a PathHashIndex AND without a FullDirectoryIndex; CreatePakFile in UnrealPak.exe has a check statement for it.
-		UE_LOG(LogPakFile, Log, TEXT("Corrupt pak PrimaryIndex detected!"));
-		UE_LOG(LogPakFile, Log, TEXT(" bReaderHasPathHashIndex: false"));
-		UE_LOG(LogPakFile, Log, TEXT(" bReaderHasFullDirectoryIndex: false"));
-		return false;
-	}
-
-	// Load the Secondary Index(es)
-	TArray<uint8> PathHashIndexData;
-	FMemoryReader PathHashIndexReader(PathHashIndexData);
-	if (bWillUsePathHashIndex)
-	{
-		FGuardedInt64 PathHashIndexEndPosition = FGuardedInt64(PathHashIndexOffset) + PathHashIndexSize;
-		if (PathHashIndexOffset < 0 || 
-			PathHashIndexSize < 0 ||
-			PathHashIndexEndPosition.InvalidOrGreaterThan(CachedTotalSize) || 
-			IntFitsIn<int32>(PathHashIndexSize) == false)
-		{
-			// Should not be possible for these values (which came from the PrimaryIndex) to be invalid, since we verified the index hash of the PrimaryIndex
-			UE_LOG(LogPakFile, Log, TEXT("Corrupt pak PrimaryIndex detected!"));
-			UE_LOG(LogPakFile, Log, TEXT(" Filename: %s"), *PakFilename);
-			UE_LOG(LogPakFile, Log, TEXT(" Total Size: %lld"), Reader.TotalSize());
-			UE_LOG(LogPakFile, Log, TEXT(" PathHashIndexOffset : %lld"), PathHashIndexOffset);
-			UE_LOG(LogPakFile, Log, TEXT(" PathHashIndexSize: %lld"), PathHashIndexSize);
-			return false;
-		}
-		Reader.Seek(PathHashIndexOffset);
-		PathHashIndexData.SetNum((int32)(PathHashIndexSize));
-		{
-			SCOPED_BOOT_TIMING("PakFile_LoadPathHashIndex");
-			Reader.Serialize(PathHashIndexData.GetData(), PathHashIndexSize);
-		}
-
-		{
-			SCOPED_BOOT_TIMING("PakFile_HashPathHashIndex");
-			if (!DecryptAndValidateIndex(Reader, PathHashIndexData, PathHashIndexHash, ComputedHash))
-			{
-				UE_LOG(LogPakFile, Log, TEXT("Corrupt pak PathHashIndex detected!"));
-				UE_LOG(LogPakFile, Log, TEXT(" Filename: %s"), *PakFilename);
-				UE_LOG(LogPakFile, Log, TEXT(" Encrypted: %d"), Info.bEncryptedIndex);
-				UE_LOG(LogPakFile, Log, TEXT(" Total Size: %lld"), Reader.TotalSize());
-				UE_LOG(LogPakFile, Log, TEXT(" Index Offset: %lld"), FullDirectoryIndexOffset);
-				UE_LOG(LogPakFile, Log, TEXT(" Index Size: %lld"), FullDirectoryIndexSize);
-				UE_LOG(LogPakFile, Log, TEXT(" Stored Index Hash: %s"), *PathHashIndexHash.ToString());
-				UE_LOG(LogPakFile, Log, TEXT(" Computed Index Hash: %s"), *ComputedHash.ToString());
-				return false;
-			}
-		}
-
-		{
-			SCOPED_BOOT_TIMING("PakFile_SerializePathHashIndex");
-			PathHashIndexReader << PathHashIndex;
-		}
-		bHasPathHashIndex = true;
-	}
-	
-	if (!bReadFullDirectoryIndex)
-	{
-		check(bWillUsePathHashIndex); // Need to confirm that we have read the PathHashIndex bytes
-		// Store the PrunedDirectoryIndex in our DirectoryIndex
-		{
-			SCOPED_BOOT_TIMING("PakFile_SerializePrunedDirectoryIndex");
-			PathHashIndexReader << DirectoryIndex;
-		}
-		bHasFullDirectoryIndex = false;
-#if ENABLE_PAKFILE_RUNTIME_PRUNING
-		bWillPruneDirectoryIndex = false;
-#endif
-	}
-	else
-	{
-		FGuardedInt64 FullDirectoryIndexEndPosition = FGuardedInt64(FullDirectoryIndexOffset) + FullDirectoryIndexSize;
-		if (FullDirectoryIndexOffset  < 0 || 
-			FullDirectoryIndexSize < 0 ||
-			FullDirectoryIndexEndPosition.InvalidOrGreaterThan(CachedTotalSize) || 
-			IntFitsIn<int32>(FullDirectoryIndexSize) == false)
-		{
-			// Should not be possible for these values (which came from the PrimaryIndex) to be invalid, since we verified the index hash of the PrimaryIndex
-			UE_LOG(LogPakFile, Log, TEXT("Corrupt pak PrimaryIndex detected!"));
-			UE_LOG(LogPakFile, Log, TEXT(" Filename: %s"), *PakFilename);
-			UE_LOG(LogPakFile, Log, TEXT(" Total Size: %lld"), Reader.TotalSize());
-			UE_LOG(LogPakFile, Log, TEXT(" FullDirectoryIndexOffset : %lld"), FullDirectoryIndexOffset );
-			UE_LOG(LogPakFile, Log, TEXT(" FullDirectoryIndexSize: %lld"), FullDirectoryIndexSize);
-			return false;
-		}
-		TArray<uint8> FullDirectoryIndexData;
-		Reader.Seek(FullDirectoryIndexOffset );
-		FullDirectoryIndexData.SetNum((int32)(FullDirectoryIndexSize));
-		{
-			SCOPED_BOOT_TIMING("PakFile_LoadDirectoryIndex");
-			Reader.Serialize(FullDirectoryIndexData.GetData(), FullDirectoryIndexSize);
-		}
-
-		{
-			SCOPED_BOOT_TIMING("PakFile_HashDirectoryIndex");
-			if (!DecryptAndValidateIndex(Reader, FullDirectoryIndexData, FullDirectoryIndexHash, ComputedHash))
-			{
-				UE_LOG(LogPakFile, Log, TEXT("Corrupt pak FullDirectoryIndex detected!"));
-				UE_LOG(LogPakFile, Log, TEXT(" Filename: %s"), *PakFilename);
-				UE_LOG(LogPakFile, Log, TEXT(" Encrypted: %d"), Info.bEncryptedIndex);
-				UE_LOG(LogPakFile, Log, TEXT(" Total Size: %lld"), Reader.TotalSize());
-				UE_LOG(LogPakFile, Log, TEXT(" Index Offset: %lld"), FullDirectoryIndexOffset);
-				UE_LOG(LogPakFile, Log, TEXT(" Index Size: %lld"), FullDirectoryIndexSize);
-				UE_LOG(LogPakFile, Log, TEXT(" Stored Index Hash: %s"), *FullDirectoryIndexHash.ToString());
-				UE_LOG(LogPakFile, Log, TEXT(" Computed Index Hash: %s"), *ComputedHash.ToString());
-				return false;
-			}
-		}
-
-		FMemoryReader SecondaryIndexReader(FullDirectoryIndexData);
-		{
-			SCOPED_BOOT_TIMING("PakFile_SerializeDirectoryIndex");
-			SecondaryIndexReader << DirectoryIndex;
-		}
-		bHasFullDirectoryIndex = true;
-
-#if ENABLE_PAKFILE_RUNTIME_PRUNING
-		if (bWillUseFullDirectoryIndex)
-		{
-			bWillPruneDirectoryIndex = false;
-		}
-		else
-		{
-			// Store the PrunedDirectoryIndex from the PrimaryIndex in our PrunedDirectoryIndex, to be used for verification and to be swapped into DirectoryIndex later
-			check(bWillUsePathHashIndex); // Need to confirm that we have read the PathHashIndex bytes
-			{
-				SCOPED_BOOT_TIMING("PakFile_SerializePrunedDirectoryIndex");
-				PathHashIndexReader << PrunedDirectoryIndex;
-			}
-			bWillPruneDirectoryIndex = true;
-			bSomePakNeedsPruning = true;
-		}
-#endif
-	}
-
-	UE_LOG(LogPakFile, Verbose, TEXT("PakFile PrimaryIndexSize=%d"), Info.IndexSize);
-	UE_LOG(LogPakFile, Verbose, TEXT("PakFile PathHashIndexSize=%d"), PathHashIndexSize);
-	UE_LOG(LogPakFile, Verbose, TEXT("PakFile FullDirectoryIndexSize=%d"), FullDirectoryIndexSize);
-
-	check(bHasFullDirectoryIndex || bHasPathHashIndex);
-	return true;
-}
-
-bool FPakFile::LoadLegacyIndex(FArchive& Reader)
-{
-	bHasPathHashIndex = false;
-	bHasFullDirectoryIndex = false;
-#if ENABLE_PAKFILE_RUNTIME_PRUNING
-	bNeedsLegacyPruning = false;
-	bWillPruneDirectoryIndex = false;
-#endif
-
-	// Load index into memory first.
-	FGuardedInt64 IndexEndPosition = FGuardedInt64(Info.IndexOffset) + Info.IndexSize;
-	if (Info.IndexSize < 0 ||
-		Info.IndexOffset < 0 ||
-		IndexEndPosition.InvalidOrGreaterThan(CachedTotalSize) ||
-		IntFitsIn<int32>(Info.IndexSize) == false)
-	{
-		UE_LOG(LogPakFile, Fatal, TEXT("Corrupted index offset/size in pak file."));
-		return false;
-	}
-
-	TArray<uint8> IndexData;
-	IndexData.SetNum((int32)(Info.IndexSize));
-
-	Reader.Seek(Info.IndexOffset);
-	Reader.Serialize(IndexData.GetData(), Info.IndexSize);
-
-	FSHAHash ComputedHash;
-	if (!DecryptAndValidateIndex(Reader, IndexData, Info.IndexHash, ComputedHash))
-	{
-		UE_LOG(LogPakFile, Log, TEXT("Corrupt pak index detected!"));
-		UE_LOG(LogPakFile, Log, TEXT(" Filename: %s"), *PakFilename);
-		UE_LOG(LogPakFile, Log, TEXT(" Encrypted: %d"), Info.bEncryptedIndex);
-		UE_LOG(LogPakFile, Log, TEXT(" Total Size: %lld"), Reader.TotalSize());
-		UE_LOG(LogPakFile, Log, TEXT(" Index Offset: %lld"), Info.IndexOffset);
-		UE_LOG(LogPakFile, Log, TEXT(" Index Size: %lld"), Info.IndexSize);
-		UE_LOG(LogPakFile, Log, TEXT(" Stored Index Hash: %s"), *Info.IndexHash.ToString());
-		UE_LOG(LogPakFile, Log, TEXT(" Computed Index Hash: %s"), *ComputedHash.ToString());
-		return false;
-	}
-
-
-	FMemoryReader IndexReader(IndexData);
-
-	// Read the default mount point and all entries.
-	NumEntries = 0;
-	IndexReader << MountPoint;
-
-	// We are just deserializing a string, which could get however long. Since we know it's a path
-	// and paths are bound to file system rules, we know it can't get absurdly long (e.g. windows is _at best_ 32k)
-	// we just sanity check it to prevent operating on massive buffers and risking overflows.
-	if (MountPoint.Len() > 65535)
-	{
-		UE_LOG(LogPakFile, Error, TEXT("Corrupt pak index data: MountPoint path is longer than 65k"));
-		return false;
-	}
-	IndexReader << NumEntries;
-
-	if (NumEntries < 0)
-	{
-		UE_LOG(LogPakFile, Error, TEXT("Corrupt pak index data: NumEntries is negative"));
-		return false;
-	}
-
-	MakeDirectoryFromPath(MountPoint);
-
-	FPakEntryPair PakEntryPair;
-	auto ReadNextEntry = [&PakEntryPair, &IndexReader, this]() -> FPakEntryPair&
-	{
-		IndexReader << PakEntryPair.Filename;
-		PakEntryPair.Info.Reset();
-		PakEntryPair.Info.Serialize(IndexReader, this->Info.Version);
-		return PakEntryPair;
-	};
-
-	TMap<uint64, FString> CollisionDetection;
-	int32 NumEncodedEntries = 0;
-	int32 NumDeletedEntries = 0;
-#if ENABLE_PAKFILE_RUNTIME_PRUNING
-	bool bCreatePathHash = !IsPakKeepFullDirectory();
-#else
-	// Pruning of legacy files is no longer supported; we will keep the entire directory with no way to prune it.  There is no need to create the PathHashIndex since we will have the FullDirectoryIndex.
-	bool bCreatePathHash = false;
-#endif
-	FPathHashIndex* PathHashToWrite = bCreatePathHash ? &PathHashIndex : nullptr;
-	FPakFile::EncodePakEntriesIntoIndex(NumEntries, ReadNextEntry, *PakFilename, Info, MountPoint, NumEncodedEntries, NumDeletedEntries, &PathHashSeed,
-		&DirectoryIndex, PathHashToWrite, EncodedPakEntries, Files, &CollisionDetection, Info.Version);
-	check(NumEncodedEntries + Files.Num() + NumDeletedEntries == NumEntries);
-	Files.Shrink();
-	EncodedPakEntries.Shrink();
-
-	bHasPathHashIndex = bCreatePathHash;
-	bHasFullDirectoryIndex = true;
-#if ENABLE_PAKFILE_RUNTIME_PRUNING
-	if (!IsPakKeepFullDirectory())
-	{
-		bNeedsLegacyPruning = true;
-		bWillPruneDirectoryIndex = true;
-		bSomePakNeedsPruning = true;
-		// We cannot prune during this call because config settings have not yet been loaded and we need the settings for DirectoryIndexKeepFiles before we can prune
-		// PrunedDirectoryIndex will be created and swapped with DirectoryIndex in OptimizeMemoryUsageForMountedPaks, and bHasFullDirectoryIndex will be set to false then
-	}
-	else
-	{
-		bNeedsLegacyPruning = false;
-		bWillPruneDirectoryIndex = false;
-	}
-#endif
-
-	check(bHasFullDirectoryIndex || bHasPathHashIndex);
-	return true;
-}
-
-bool FPakFile::DecryptAndValidateIndex(FArchive& Reader, TArray<uint8>& IndexData, FSHAHash& InExpectedHash, FSHAHash& OutActualHash)
-{
-	// Decrypt if necessary
-	if (Info.bEncryptedIndex)
-	{
-		DecryptData(IndexData.GetData(), IndexData.Num(), Info.EncryptionKeyGuid);
-	}
-
-	// Check SHA1 value.
-	FSHA1::HashBuffer(IndexData.GetData(), IndexData.Num(), OutActualHash.Hash);
-	return InExpectedHash == OutActualHash;
-}
-
-/*** This is a copy of FFnv::MemFnv64 from before the bugfix for swapped Offset and Prime. It is used to decode legacy paks that have hashes created from the prebugfix version of the function */
-static uint64 LegacyMemFnv64(const void* InData, int32 Length, uint64 InOffset)
-{
-	// constants from above reference
-	static const uint64 Offset = 0x00000100000001b3;
-	static const uint64 Prime = 0xcbf29ce484222325;
-
-	const uint8* __restrict Data = (uint8*)InData;
-
-	uint64 Fnv = Offset + InOffset; // this is not strictly correct as the offset should be prime and InOffset could be arbitrary
-	for (; Length; --Length)
-	{
-		Fnv ^= *Data++;
-		Fnv *= Prime;
-	}
-
-	return Fnv;
-}
-
-uint64 FPakFile::HashPath(const TCHAR* RelativePathFromMount, uint64 Seed, int32 PakFileVersion)
-{
-	FString LowercaseRelativePath(RelativePathFromMount);
-	LowercaseRelativePath.ToLowerInline();
-	if (PakFileVersion >= FPakInfo::PakFile_Version_Fnv64BugFix)
-	{
-		return FFnv::MemFnv64(*LowercaseRelativePath, LowercaseRelativePath.Len() * sizeof(TCHAR), Seed);
-	}
-	else
-	{
-		return LegacyMemFnv64(*LowercaseRelativePath, LowercaseRelativePath.Len() * sizeof(TCHAR), Seed);
-	}
-}
-
-void FPakFile::EncodePakEntriesIntoIndex(int32 InNumEntries, const ReadNextEntryFunction& InReadNextEntry, const TCHAR* InPakFilename, const FPakInfo& InPakInfo, const FString& MountPoint,
-	int32& OutNumEncodedEntries, int32& OutNumDeletedEntries, uint64* OutPathHashSeed,
-	FDirectoryIndex* OutDirectoryIndex, FPathHashIndex* OutPathHashIndex, TArray<uint8>& OutEncodedPakEntries, TArray<FPakEntry>& OutNonEncodableEntries,
-	TMap<uint64, FString>* InOutCollisionDetection, int32 PakFileVersion)
-{
-	uint64 PathHashSeed = 0;
-	if (OutPathHashSeed || OutPathHashIndex)
-	{
-		FString LowercasePakFilename(InPakFilename);
-		LowercasePakFilename.ToLowerInline();
-		PathHashSeed = FCrc::StrCrc32(*LowercasePakFilename);
-		if (OutPathHashSeed)
-		{
-			*OutPathHashSeed = PathHashSeed;
-		}
-	}
-
-	OutNumEncodedEntries = 0;
-	OutNumDeletedEntries = 0;
-	FMemoryWriter CompressedEntryWriter(OutEncodedPakEntries);
-	for (int32 EntryCount = 0; EntryCount < InNumEntries; ++EntryCount)
-	{
-		FPakEntryPair& Pair = InReadNextEntry();
-		// Add the Entry and get an FPakEntryLocation for it
-		const FPakEntry& PakEntry = Pair.Info;
-		FPakEntryLocation EntryLocation;
-		if (!PakEntry.IsDeleteRecord())
-		{
-			EntryLocation = FPakEntryLocation::CreateFromOffsetIntoEncoded(OutEncodedPakEntries.Num());
-			if (EncodePakEntry(CompressedEntryWriter, PakEntry, InPakInfo))
-			{
-				++OutNumEncodedEntries;
-			}
-			else
-			{
-				int32 ListIndex = OutNonEncodableEntries.Num();
-				EntryLocation = FPakEntryLocation::CreateFromListIndex(ListIndex);
-				OutNonEncodableEntries.Add(PakEntry);
-
-				// PakEntries in the index have some values that are different from the in-place pakentries stored next to each file's payload.  EncodePakEntry handles that internally if encoding succeeded.
-				FPakEntry& StoredPakEntry = OutNonEncodableEntries[ListIndex];
-				FMemory::Memset(StoredPakEntry.Hash, 0); // Hash is 0-filled
-				StoredPakEntry.Verified = true; // Validation of the hash is impossible, so Verified is set to true
-			}
-		}
-		else
-		{
-			++OutNumDeletedEntries;
-		}
-
-		// Add the Entry into the requested Indexes
-		AddEntryToIndex(Pair.Filename, EntryLocation, MountPoint, PathHashSeed, OutDirectoryIndex, OutPathHashIndex, InOutCollisionDetection, PakFileVersion);
-	}
-}
-
-void FPakFile::PruneDirectoryIndex(FDirectoryIndex& InOutDirectoryIndex, FDirectoryIndex* PrunedDirectoryIndex, const FString& MountPoint)
-{
-	// Caller holds WriteLock on DirectoryIndexLock
-	TArray<FString> FileWildCards, DirectoryWildCards, OldWildCards;
-	GConfig->GetArray(TEXT("Pak"), TEXT("DirectoryIndexKeepFiles"), FileWildCards, GEngineIni);
-	GConfig->GetArray(TEXT("Pak"), TEXT("DirectoryIndexKeepEmptyDirectories"), DirectoryWildCards, GEngineIni);
-	GConfig->GetArray(TEXT("Pak"), TEXT("DirectoryRootsToKeepInMemoryWhenUnloadingPakEntryFilenames"), OldWildCards, GEngineIni); // Legacy name, treated as both KeepFiles and KeepEmptyDirectories
-	DirectoryWildCards.Append(OldWildCards);
-	FileWildCards.Append(OldWildCards);
-	int32 NumKeptEntries = 0;
-
-	if (PrunedDirectoryIndex)
-	{
-		PrunedDirectoryIndex->Empty();
-	}
-
-	TMap<FString, bool> KeepDirectory;
-
-	// Clear out those portions of the Index allowed by the user.
-	if (DirectoryWildCards.Num() > 0 || FileWildCards.Num() > 0)
-	{
-		for (auto It = InOutDirectoryIndex.CreateIterator(); It; ++It)
-		{
-			const FString& DirectoryName = It.Key();
-			FPakDirectory& OriginalDirectory = It.Value();
-			const FString FullDirectoryName = PakPathCombine(MountPoint, DirectoryName);
-			check(IsPathInDirectoryFormat(FullDirectoryName));
-			FPakDirectory* PrunedDirectory = nullptr;
-			bool bKeepDirectory = false;
-
-			TArray<FString> RemoveFilenames;
-			for (auto FileIt = It->Value.CreateIterator(); FileIt; ++FileIt)
-			{
-				const FString& FileNameWithoutPath = FileIt->Key;
-				const FString FullFilename = PakPathCombine(FullDirectoryName, FileNameWithoutPath);
-				bool bKeepFile = false;
-
-				for (const FString& WildCard : FileWildCards)
-				{
-					if (FullFilename.MatchesWildcard(WildCard))
-					{
-						bKeepFile = true;
-						break;
-					}
-				}
-
-				if (bKeepFile)
-				{
-					bKeepDirectory = true;
-					if (PrunedDirectoryIndex)
-					{
-						if (!PrunedDirectory)
-						{
-							PrunedDirectory = &PrunedDirectoryIndex->Add(DirectoryName);
-						}
-						PrunedDirectory->Add(FileNameWithoutPath, FileIt->Value);
-					}
-				}
-				else
-				{
-					if (!PrunedDirectoryIndex)
-					{
-						RemoveFilenames.Add(FileNameWithoutPath);
-					}
-				}
-			}
-			if (!PrunedDirectoryIndex)
-			{
-				for (const FString& FileNameWithoutPath : RemoveFilenames)
-				{
-					OriginalDirectory.Remove(FileNameWithoutPath);
-				}
-			}
-
-			if (!bKeepDirectory)
-			{
-				for (const FString& WildCard : DirectoryWildCards)
-				{
-					if (FullDirectoryName.MatchesWildcard(WildCard))
-					{
-						bKeepDirectory = true;
-						break;
-					}
-				}
-			}
-			KeepDirectory.FindOrAdd(DirectoryName) = bKeepDirectory;
-		}
-
-		// For each kept directory, mark that we need to keep all of its parents up to the mount point
-		for (const TPair<FString, bool>& Pair : KeepDirectory)
-		{
-			if (Pair.Value)
-			{
-				FString CurrentDirectory = Pair.Key;
-				FString UnusedCleanFileName;
-				while (CurrentDirectory != MountPoint)
-				{
-					if (!SplitPathInline(CurrentDirectory, UnusedCleanFileName))
-					{
-						break;
-					}
-					bool& bOldValue = KeepDirectory.FindOrAdd(CurrentDirectory);
-					if (bOldValue)
-					{
-						break;
-					}
-					bOldValue = true;
-				}
-			}
-		}
-
-		// Prune all of the directories
-		for (const TPair<FString, bool>& Pair : KeepDirectory)
-		{
-			const FString& DirectoryName = Pair.Key;
-			bool bKeep = Pair.Value;
-			if (bKeep)
-			{
-				if (PrunedDirectoryIndex)
-				{
-					PrunedDirectoryIndex->FindOrAdd(DirectoryName);
-				}
-			}
-			else
-			{
-				if (!PrunedDirectoryIndex)
-				{
-					InOutDirectoryIndex.Remove(DirectoryName);
-				}
-			}
-		}
-	}
-	else
-	{
-		if (!PrunedDirectoryIndex)
-		{
-			InOutDirectoryIndex.Empty();
-		}
-	}
-}
-
-FPakFile::EFindResult FPakFile::GetPakEntry(const FPakEntryLocation& PakEntryLocation, FPakEntry* OutEntry) const
-{
-	return GetPakEntry(PakEntryLocation, OutEntry, EncodedPakEntries, Files, Info);
-}
-
-FPakFile::EFindResult FPakFile::GetPakEntry(const FPakEntryLocation& PakEntryLocation, FPakEntry* OutEntry, const TArray<uint8>& EncodedPakEntries, const TArray<FPakEntry>& Files, const FPakInfo& Info)
-{
-	bool bDeleted = PakEntryLocation.IsInvalid();
-	if (OutEntry != NULL)
-	{
-		if (!bDeleted)
-		{
-			// The FPakEntry structures are bit-encoded, so decode it.
-			int32 EncodedOffset = PakEntryLocation.GetAsOffsetIntoEncoded();
-			if (EncodedOffset >= 0)
-			{
-				check(EncodedOffset < EncodedPakEntries.Num());
-				DecodePakEntry(EncodedPakEntries.GetData() + EncodedOffset, *OutEntry, Info);
-			}
-			else
-			{
-				int32 ListIndex = PakEntryLocation.GetAsListIndex();
-				check(ListIndex >= 0);
-				const FPakEntry& FoundEntry = Files[ListIndex];
-				//*OutEntry = **FoundEntry;
-				OutEntry->Offset = FoundEntry.Offset;
-				OutEntry->Size = FoundEntry.Size;
-				OutEntry->UncompressedSize = FoundEntry.UncompressedSize;
-				OutEntry->CompressionMethodIndex = FoundEntry.CompressionMethodIndex;
-				OutEntry->CompressionBlocks = FoundEntry.CompressionBlocks;
-				OutEntry->CompressionBlockSize = FoundEntry.CompressionBlockSize;
-				OutEntry->Flags = FoundEntry.Flags;
-			}
-		}
-		else
-		{
-			// entry was deleted, build dummy entry to indicate the deleted entry
-			(*OutEntry) = FPakEntry();
-			OutEntry->SetDeleteRecord(true);
-		}
-
-		// Index PakEntries do not store their hash, so verification against the hash is impossible.
-		// Initialize the OutEntry's Hash to 0 and mark it as already verified.
-		// TODO: Verified and Hash are checked by FPakFileHandle, which is used when opening files from PakFiles synchronously;
-		//       it is not currently used by asynchronous reads in FPakPrecacher, and we can likely remove it from use in FPakFileHandle as well
-		FMemory::Memset(OutEntry->Hash, 0);
-		OutEntry->Verified = true;
-	}
-
-	return bDeleted ? EFindResult::FoundDeleted : EFindResult::Found;
-}
-
-struct FPakFile::FIndexSettings
-{
-	FIndexSettings()
-	{
-		bKeepFullDirectory = true;
-		bValidatePruning = false;
-		bDelayPruning = false;
-		bWritePathHashIndex = true;
-		bWriteFullDirectoryIndex = true;
-
-		// Paks are mounted before config files are read, so the licensee needs to hardcode all settings used for runtime index loading rather than specifying them in ini
-		if (FPakPlatformFile::GetPakSetIndexSettingsDelegate().IsBound())
-		{
-			FPakPlatformFile::GetPakSetIndexSettingsDelegate().Execute(bKeepFullDirectory, bValidatePruning, bDelayPruning);
-		}
-
-		// Settings not used at runtime can be read from ini
-#if !UE_BUILD_SHIPPING
-		GConfig->GetBool(TEXT("Pak"), TEXT("WritePathHashIndex"), bWritePathHashIndex, GEngineIni);
-		GConfig->GetBool(TEXT("Pak"), TEXT("WriteFullDirectoryIndex"), bWriteFullDirectoryIndex, GEngineIni);
-#endif
-
-#if IS_PROGRAM || WITH_EDITOR
-		// Directory pruning is not enabled in the editor or in development programs because there is no need to save the memory in those environments and some development features require not pruning
-		bKeepFullDirectory = true;
-#else
-		bKeepFullDirectory = bKeepFullDirectory || !FPlatformProperties::RequiresCookedData();
-#endif
-#if !UE_BUILD_SHIPPING
-		const TCHAR* CommandLine = FCommandLine::Get();
-		FParse::Bool(CommandLine, TEXT("ForcePakKeepFullDirectory="), bKeepFullDirectory);
-#if ENABLE_PAKFILE_RUNTIME_PRUNING_VALIDATE
-		FParse::Bool(CommandLine, TEXT("ForcePakValidatePruning="), bValidatePruning);
-		FParse::Bool(CommandLine, TEXT("ForcePakDelayPruning="), bDelayPruning);
-#endif
-		FParse::Bool(CommandLine, TEXT("ForcePakWritePathHashIndex="), bWritePathHashIndex);
-		FParse::Bool(CommandLine, TEXT("ForcePakWriteFullDirectoryIndex="), bWriteFullDirectoryIndex);
-#endif
-	}
-
-	bool bKeepFullDirectory;
-	bool bValidatePruning;
-	bool bDelayPruning;
-	bool bWritePathHashIndex;
-	bool bWriteFullDirectoryIndex;
-};
-
-FPakFile::FIndexSettings& FPakFile::GetIndexSettings()
-{
-	static FIndexSettings IndexLoadParams;
-	return IndexLoadParams;
-}
-
-bool FPakFile::IsPakKeepFullDirectory()
-{
-	FIndexSettings& IndexLoadParams = GetIndexSettings();
-	return IndexLoadParams.bKeepFullDirectory;
-}
-
-bool FPakFile::IsPakValidatePruning()
-{
-#if ENABLE_PAKFILE_RUNTIME_PRUNING_VALIDATE
-	FIndexSettings& IndexLoadParams = GetIndexSettings();
-	return IndexLoadParams.bValidatePruning;
-#else
-	return false;
-#endif
-}
-
-bool FPakFile::IsPakDelayPruning()
-{
-	FIndexSettings& IndexLoadParams = GetIndexSettings();
-	return IndexLoadParams.bDelayPruning;
-}
-
-bool FPakFile::IsPakWritePathHashIndex()
-{
-	FIndexSettings& IndexLoadParams = GetIndexSettings();
-	return IndexLoadParams.bWritePathHashIndex;
-}
-
-bool FPakFile::IsPakWriteFullDirectoryIndex()
-{
-	FIndexSettings& IndexLoadParams = GetIndexSettings();
-	return IndexLoadParams.bWriteFullDirectoryIndex;
-}
-
-bool FPakFile::RequiresDirectoryIndexLock() const
-{
-#if ENABLE_PAKFILE_RUNTIME_PRUNING
-	return bWillPruneDirectoryIndex;
-#else
-	return false; 
-#endif
-}
-
-bool FPakFile::ShouldValidatePrunedDirectory() const
-{
-#if ENABLE_PAKFILE_RUNTIME_PRUNING_VALIDATE
-	return IsPakValidatePruning() && bWillPruneDirectoryIndex && !bNeedsLegacyPruning;
-#else
-	return false;
-#endif
-}
-
-
-void FPakFile::AddEntryToIndex(const FString& Filename, const FPakEntryLocation& EntryLocation, const FString& MountPoint, uint64 PathHashSeed,
-	FDirectoryIndex* DirectoryIndex, FPathHashIndex* PathHashIndex, TMap<uint64, FString>* CollisionDetection, int32 PakFileVersion)
-{
-	FString RelativePathFromMount;
-	if (FPaths::IsRelative(Filename))
-	{
-		RelativePathFromMount = Filename;
-	}
-	else
-	{
-		check(IsPathInDirectoryFormat(MountPoint));
-		RelativePathFromMount = Filename;
-		check(Filename.Len() > MountPoint.Len());
-		bool bSucceeded = GetRelativePathFromMountInline(RelativePathFromMount, MountPoint);
-		check(bSucceeded);
-	}
-
-	if (DirectoryIndex)
-	{
-		FString RelativeDirectoryFromMount = RelativePathFromMount;
-		FString CleanFileName;
-		SplitPathInline(RelativeDirectoryFromMount, CleanFileName);
-		FPakDirectory* Directory = DirectoryIndex->Find(RelativeDirectoryFromMount);
-		if (Directory == nullptr)
-		{
-			// add the parent directories up to the mount point (mount point relative path is "/")
-			FString CurrentDirectory(RelativeDirectoryFromMount);
-			FString UnusedCleanFileName;
-			while (!CurrentDirectory.IsEmpty())
-			{
-				// FPaths::GetPath doesn't handle our / at the end of directories, so call FindLastChar(/)
-				if (!SplitPathInline(CurrentDirectory, UnusedCleanFileName))
-				{
-					break;
-				}
-				DirectoryIndex->FindOrAdd(CurrentDirectory);
-			}
-
-			// Add the new directory; this has to come after the addition of the parent directories because we want to use the pointer afterwards and adding other directories would invalidate it
-			Directory = &DirectoryIndex->Add(RelativeDirectoryFromMount);
-		}
-		Directory->Add(CleanFileName, EntryLocation);
-	}
-
-	// Add the entry into the PathHash index
-	if (CollisionDetection || PathHashIndex)
-	{
-		uint64 PathHash = FPakFile::HashPath(*RelativePathFromMount, PathHashSeed, PakFileVersion);
-		if (CollisionDetection)
-		{
-			FString* ExistingFilename = CollisionDetection->Find(PathHash);
-			checkf(!ExistingFilename || ExistingFilename->Equals(RelativePathFromMount), TEXT("Hash collision for two Filenames within a PakFile.  Filename1 == '%s'.  Filename2 == '%s'.  Hash='0x%x'.")
-				TEXT(" Collision handling is not yet implemented; to resolve the conflict, modify one of the Filenames."), **ExistingFilename, *RelativePathFromMount, PathHash);
-			CollisionDetection->Add(PathHash, RelativePathFromMount);
-		}
-		if (PathHashIndex)
-		{
-			PathHashIndex->Add(PathHash, EntryLocation);
-		}
-	}
-}
-
-// Take a pak entry and byte encode it into a smaller representation
-bool FPakFile::EncodePakEntry(FArchive& Ar, const FPakEntry& InPakEntry, const FPakInfo& InInfo)
-{
-	// See notes in DecodePakEntry for the output layout
-
-	check(Ar.IsSaving()); // This function is encode only, and promises not to modify InEntry, but we want to use << which takes non-const so we need to assert that Ar is a loader not a saver
-	check(!InPakEntry.IsDeleteRecord()); // Deleted PakEntries cannot be encoded, caller must check for IsDeleteRecord and handle it separately by e.g. not adding the Entry to the FileList and instead giving the referencer an invalid FPakEntryLocation
-	FPakEntry& PakEntry = const_cast<FPakEntry&>(InPakEntry);
-
-	const uint32 CompressedBlockAlignment = PakEntry.IsEncrypted() ? FAES::AESBlockSize : 1;
-	const int64 HeaderSize = PakEntry.GetSerializedSize(InInfo.Version);
-
-	// This data fits into a bitfield (described in DecodePakEntry), and the data has
-	// to fit within a certain range of bits.
-	if (PakEntry.CompressionMethodIndex >= (1 << 6))
-	{
-		return false;
-	}
-	if (PakEntry.CompressionBlocks.Num() >= (1 << 16))
-	{
-		return false;
-	}
-	if (PakEntry.CompressionMethodIndex != 0)
-	{
-		if (PakEntry.CompressionBlocks.Num() > 0 && ((InInfo.HasRelativeCompressedChunkOffsets() ? 0 : PakEntry.Offset) + HeaderSize != PakEntry.CompressionBlocks[0].CompressedStart))
-		{
-			return false;
-		}
-		if (PakEntry.CompressionBlocks.Num() == 1)
-		{
-			uint64 Base = InInfo.HasRelativeCompressedChunkOffsets() ? 0 : PakEntry.Offset;
-			uint64 AlignedBlockSize = Align(PakEntry.CompressionBlocks[0].CompressedEnd - PakEntry.CompressionBlocks[0].CompressedStart, CompressedBlockAlignment);
-			if ((Base + HeaderSize + PakEntry.Size) != (PakEntry.CompressionBlocks[0].CompressedStart + AlignedBlockSize))
-			{
-				return false;
-			}
-		}
-		if (PakEntry.CompressionBlocks.Num() > 1)
-		{
-			for (int i = 1; i < PakEntry.CompressionBlocks.Num(); ++i)
-			{
-				uint64 PrevBlockSize = PakEntry.CompressionBlocks[i - 1].CompressedEnd - PakEntry.CompressionBlocks[i - 1].CompressedStart;
-				PrevBlockSize = Align(PrevBlockSize, CompressedBlockAlignment);
-				if (PakEntry.CompressionBlocks[i].CompressedStart != (PakEntry.CompressionBlocks[i - 1].CompressedStart + PrevBlockSize))
-				{
-					return false;
-				}
-			}
-		}
-	}
-
-	// This entry can be encoded, so let's do it!
-
-	bool bIsOffset32BitSafe = PakEntry.Offset <= MAX_uint32;
-	bool bIsSize32BitSafe = PakEntry.Size <= MAX_uint32;
-	bool bIsUncompressedSize32BitSafe = PakEntry.UncompressedSize <= MAX_uint32;
-	
-	// If CompressionBlocks.Num() == 1, we set CompressionBlockSize == UncompressedSize and record CompressBlockSizePacked=0
-	// Otherwise, we encode CompressionBlockSize as a 6-bit multiple of 1 << 11.
-	// If CompressionBlockSize is not a multiple of 1 << 11, or is a larger multiple than 6 bits we can not encode correctly.
-	// In that case we set the packed field to its maximum value (0x3F) and send the unencoded CompressionBlockSize as a separate value.
-	uint32 CompressionBlockSizePacked = 0;
-	if (PakEntry.CompressionBlocks.Num() > 1)
-	{
-		CompressionBlockSizePacked = (PakEntry.CompressionBlockSize >> 11) & 0x3F;
-		if ((CompressionBlockSizePacked << 11) != PakEntry.CompressionBlockSize)
-		{
-			CompressionBlockSizePacked = 0x3F;
-		}
-	}
-
-	// Build the Flags field.
-	uint32 Flags =
-		(bIsOffset32BitSafe ? (1 << 31) : 0)
-		| (bIsUncompressedSize32BitSafe ? (1 << 30) : 0)
-		| (bIsSize32BitSafe ? (1 << 29) : 0)
-		| (PakEntry.CompressionMethodIndex << 23)
-		| (PakEntry.IsEncrypted() ? (1 << 22) : 0)
-		| (PakEntry.CompressionBlocks.Num() << 6)
-		| CompressionBlockSizePacked
-		;
-
-	Ar << Flags;
-	
-	// if we write 0x3F for CompressionBlockSize then send the field
-	if ( CompressionBlockSizePacked == 0x3F )
-	{
-		uint32 Value = (uint32)PakEntry.CompressionBlockSize;
-		Ar << Value;
-	}
-
-	// Build the Offset field.
-	if (bIsOffset32BitSafe)
-	{
-		uint32 Value = (uint32)PakEntry.Offset;
-		Ar << Value;
-	}
-	else
-	{
-		Ar << PakEntry.Offset;
-	}
-
-	// Build the Uncompressed Size field.
-	if (bIsUncompressedSize32BitSafe)
-	{
-		uint32 Value = (uint32)PakEntry.UncompressedSize;
-		Ar << Value;
-	}
-	else
-	{
-		Ar << PakEntry.UncompressedSize;
-	}
-
-	// Any additional data is for compressed file data.
-	if (PakEntry.CompressionMethodIndex != 0)
-	{
-		// Build the Compressed Size field.
-		if (bIsSize32BitSafe)
-		{
-			uint32 Value = (uint32)PakEntry.Size;
-			Ar << Value;
-		}
-		else
-		{
-			Ar << PakEntry.Size;
-		}
-
-		// Build the Compression Blocks array.
-		if (PakEntry.CompressionBlocks.Num() > 1 || (PakEntry.CompressionBlocks.Num() == 1 && PakEntry.IsEncrypted()))
-		{
-			for (int CompressionBlockIndex = 0; CompressionBlockIndex < PakEntry.CompressionBlocks.Num(); ++CompressionBlockIndex)
-			{
-				uint32 Value = IntCastChecked<uint32>(PakEntry.CompressionBlocks[CompressionBlockIndex].CompressedEnd - PakEntry.CompressionBlocks[CompressionBlockIndex].CompressedStart);
-				Ar << Value;
-			}
-		}
-	}
-
-	return true;
-}
-
-void FPakFile::DecodePakEntry(const uint8* SourcePtr, FPakEntry& OutEntry, const FPakInfo& InInfo)
-{
-	// Grab the big bitfield value:
-	// Bit 31 = Offset 32-bit safe?
-	// Bit 30 = Uncompressed size 32-bit safe?
-	// Bit 29 = Size 32-bit safe?
-	// Bits 28-23 = Compression method
-	// Bit 22 = Encrypted
-	// Bits 21-6 = Compression blocks count
-	// Bits 5-0 = Compression block size
-	uint32 Value = *(uint32*)SourcePtr;
-	SourcePtr += sizeof(uint32);
-	
-	uint32 CompressionBlockSize = 0;
-	if ( (Value & 0x3f) == 0x3f ) // flag value to load a field
-	{
-		CompressionBlockSize = *(uint32*)SourcePtr;
-		SourcePtr += sizeof(uint32);
-	}
-	else
-	{
-		// for backwards compatibility with old paks :
-		CompressionBlockSize = ((Value & 0x3f) << 11);
-	}
-
-	// Filter out the CompressionMethod.
-	OutEntry.CompressionMethodIndex = (Value >> 23) & 0x3f;
-
-	// Test for 32-bit safe values. Grab it, or memcpy the 64-bit value
-	// to avoid alignment exceptions on platforms requiring 64-bit alignment
-	// for 64-bit variables.
-	//
-	// Read the Offset.
-	bool bIsOffset32BitSafe = (Value & (1 << 31)) != 0;
-	if (bIsOffset32BitSafe)
-	{
-		OutEntry.Offset = *(uint32*)SourcePtr;
-		SourcePtr += sizeof(uint32);
-	}
-	else
-	{
-		FMemory::Memcpy(&OutEntry.Offset, SourcePtr, sizeof(int64));
-		SourcePtr += sizeof(int64);
-	}
-
-	// Read the UncompressedSize.
-	bool bIsUncompressedSize32BitSafe = (Value & (1 << 30)) != 0;
-	if (bIsUncompressedSize32BitSafe)
-	{
-		OutEntry.UncompressedSize = *(uint32*)SourcePtr;
-		SourcePtr += sizeof(uint32);
-	}
-	else
-	{
-		FMemory::Memcpy(&OutEntry.UncompressedSize, SourcePtr, sizeof(int64));
-		SourcePtr += sizeof(int64);
-	}
-
-	// Fill in the Size.
-	if (OutEntry.CompressionMethodIndex != 0)
-	{
-		// Size is only present if compression is applied.
-		bool bIsSize32BitSafe = (Value & (1 << 29)) != 0;
-		if (bIsSize32BitSafe)
-		{
-			OutEntry.Size = *(uint32*)SourcePtr;
-			SourcePtr += sizeof(uint32);
-		}
-		else
-		{
-			FMemory::Memcpy(&OutEntry.Size, SourcePtr, sizeof(int64));
-			SourcePtr += sizeof(int64);
-		}
-	}
-	else
-	{
-		// The Size is the same thing as the UncompressedSize when
-		// CompressionMethod == COMPRESS_None.
-		OutEntry.Size = OutEntry.UncompressedSize;
-	}
-
-	// Filter the encrypted flag.
-	OutEntry.SetEncrypted((Value & (1 << 22)) != 0);
-
-	// This should clear out any excess CompressionBlocks that may be valid in the user's
-	// passed in entry.
-	uint32 CompressionBlocksCount = (Value >> 6) & 0xffff;
-	OutEntry.CompressionBlocks.Empty(CompressionBlocksCount);
-	OutEntry.CompressionBlocks.SetNum(CompressionBlocksCount);
-	
-	OutEntry.CompressionBlockSize = 0;
-	if (CompressionBlocksCount > 0)
-	{
-		OutEntry.CompressionBlockSize = CompressionBlockSize;
-		// Per the comment in Encode, if CompressionBlocksCount == 1, we use UncompressedSize for CompressionBlockSize
-		if (CompressionBlocksCount == 1)
-		{
-			OutEntry.CompressionBlockSize = IntCastChecked<uint32>(OutEntry.UncompressedSize);
-		}
-		ensure(OutEntry.CompressionBlockSize != 0);
-	}
-
-	// Set bDeleteRecord to false, because it obviously isn't deleted if we are here.
-	OutEntry.SetDeleteRecord(false);
-
-	// Base offset to the compressed data
-	int64 BaseOffset = InInfo.HasRelativeCompressedChunkOffsets() ? 0 : OutEntry.Offset;
-
-	// Handle building of the CompressionBlocks array.
-	if (OutEntry.CompressionBlocks.Num() == 1 && !OutEntry.IsEncrypted())
-	{
-		// If the number of CompressionBlocks is 1, we didn't store any extra information.
-		// Derive what we can from the entry's file offset and size.
-		FPakCompressedBlock& CompressedBlock = OutEntry.CompressionBlocks[0];
-		CompressedBlock.CompressedStart = BaseOffset + OutEntry.GetSerializedSize(InInfo.Version);
-		CompressedBlock.CompressedEnd = CompressedBlock.CompressedStart + OutEntry.Size;
-	}
-	else if (OutEntry.CompressionBlocks.Num() > 0)
-	{
-		// Get the right pointer to start copying the CompressionBlocks information from.
-		uint32* CompressionBlockSizePtr = (uint32*)SourcePtr;
-
-		// Alignment of the compressed blocks
-		uint64 CompressedBlockAlignment = OutEntry.IsEncrypted() ? FAES::AESBlockSize : 1;
-
-		// CompressedBlockOffset is the starting offset. Everything else can be derived from there.
-		int64 CompressedBlockOffset = BaseOffset + OutEntry.GetSerializedSize(InInfo.Version);
-		for (int CompressionBlockIndex = 0; CompressionBlockIndex < OutEntry.CompressionBlocks.Num(); ++CompressionBlockIndex)
-		{
-			FPakCompressedBlock& CompressedBlock = OutEntry.CompressionBlocks[CompressionBlockIndex];
-			CompressedBlock.CompressedStart = CompressedBlockOffset;
-			CompressedBlock.CompressedEnd = CompressedBlockOffset + *CompressionBlockSizePtr++;
-			CompressedBlockOffset += Align(CompressedBlock.CompressedEnd - CompressedBlock.CompressedStart, CompressedBlockAlignment);
-		}
-	}
-}
-
-bool FPakFile::NormalizeDirectoryQuery(const TCHAR* InPath, FString& OutRelativePathFromMount) const
-{
-	OutRelativePathFromMount = InPath;
-	MakeDirectoryFromPath(OutRelativePathFromMount);
-	return GetRelativePathFromMountInline(OutRelativePathFromMount, MountPoint);
-}
-
-const FPakDirectory* FPakFile::FindPrunedDirectoryInternal(const FString& RelativePathFromMount) const
-{
-	// Caller holds FScopedPakDirectoryIndexAccess
-	const FPakDirectory* PakDirectory = nullptr;
-
-#if ENABLE_PAKFILE_RUNTIME_PRUNING_VALIDATE
-	if (ShouldValidatePrunedDirectory())
-	{
-		PakDirectory = DirectoryIndex.Find(RelativePathFromMount);
-		const FPakDirectory* PrunedPakDirectory = PrunedDirectoryIndex.Find(RelativePathFromMount);
-		if ((PakDirectory != nullptr) != (PrunedPakDirectory != nullptr))
-		{
-			TSet<FString> FullFoundFiles, PrunedFoundFiles;
-			FString ReportedDirectoryName = MountPoint + RelativePathFromMount;
-			if (PakDirectory)
-			{
-				FullFoundFiles.Add(ReportedDirectoryName);
-			}
-			if (PrunedPakDirectory)
-			{
-				PrunedFoundFiles.Add(ReportedDirectoryName);
-			}
-			ValidateDirectorySearch(FullFoundFiles, PrunedFoundFiles, *ReportedDirectoryName);
-		}
-	}
-	else
-#endif
-	{
-		PakDirectory = DirectoryIndex.Find(RelativePathFromMount);
-	}
-	return PakDirectory;
-}
-
-bool FPakFile::Check()
-{
-	UE_LOG(LogPakFile, Display, TEXT("Checking pak file \"%s\". This may take a while..."), *PakFilename);
-	double StartTime = FPlatformTime::Seconds();
-
-	FSharedPakReader PakReader = GetSharedReader(nullptr);
-	int32 ErrorCount = 0;
-	int32 FileCount = 0;
-
-	// If the pak file is signed, we can do a fast check by just reading a single byte from the start of
-	// each signing block. The signed archive reader will bring in that whole signing block and compare
-	// against the signature table and fire the handler
-	if (bSigned)
-	{
-		FDelegateHandle DelegateHandle;
-		FPakPlatformFile::FPakSigningFailureHandlerData& HandlerData = FPakPlatformFile::GetPakSigningFailureHandlerData();
-
-		{
-			FScopeLock Lock(&HandlerData.GetLock());
-			DelegateHandle = HandlerData.GetPakChunkSignatureCheckFailedDelegate().AddLambda([&ErrorCount](const FPakChunkSignatureCheckFailedData&)
-			{
-				++ErrorCount;
-			});
-		}
-
-		int64 CurrentPos = 0;
-		const int64 Size = PakReader->TotalSize();
-		while (CurrentPos < Size)
-		{
-			PakReader->Seek(CurrentPos);
-			uint8 Byte = 0;
-			PakReader.GetArchive() << Byte;
-			CurrentPos += FPakInfo::MaxChunkDataSize;
-		}
-
-		if (DelegateHandle.IsValid())
-		{
-			FScopeLock Lock(&HandlerData.GetLock());
-			HandlerData.GetPakChunkSignatureCheckFailedDelegate().Remove(DelegateHandle);
-		}
-	}
-	else
-	{
-		const bool bIncludeDeleted = true;
-		TCHAR EntryNameBuffer[256];
-		auto GetEntryName = [&EntryNameBuffer](const FPakFile::FPakEntryIterator& It)
-		{
-			const FString* EntryFilename = It.TryGetFilename();
-			if (EntryFilename)
-			{
-				TCString<TCHAR>::Snprintf(EntryNameBuffer, sizeof(EntryNameBuffer), TEXT("\"%s\""), **EntryFilename);
-			}
-			else
-			{
-				TCString<TCHAR>::Snprintf(EntryNameBuffer, sizeof(EntryNameBuffer), TEXT("file at offset %u"), It.Info().Offset);
-			}
-			return EntryNameBuffer;
-		};
-		for (FPakFile::FPakEntryIterator It(*this, bIncludeDeleted); It; ++It, ++FileCount)
-		{
-			const FPakEntry& EntryFromIndex = It.Info();
-			if (EntryFromIndex.IsDeleteRecord())
-			{
-				UE_LOG(LogPakFile, Verbose, TEXT("%s Deleted."), GetEntryName(It));
-				continue;
-			}
-
-			void* FileContents = FMemory::Malloc(EntryFromIndex.Size);
-			PakReader->Seek(EntryFromIndex.Offset);
-			uint32 SerializedCrcTest = 0;
-			FPakEntry EntryFromPayload;
-			EntryFromPayload.Serialize(PakReader.GetArchive(), GetInfo().Version);
-			if (!EntryFromPayload.IndexDataEquals(EntryFromIndex))
-			{
-				UE_LOG(LogPakFile, Error, TEXT("Index FPakEntry does not match Payload FPakEntry for %s."), GetEntryName(It));
-				ErrorCount++;
-			}
-			PakReader->Serialize(FileContents, EntryFromIndex.Size);
-
-			uint8 TestHash[20];
-			FSHA1::HashBuffer(FileContents, EntryFromIndex.Size, TestHash);
-			if (FMemory::Memcmp(TestHash, EntryFromPayload.Hash, sizeof(TestHash)) != 0)
-			{
-				UE_LOG(LogPakFile, Error, TEXT("Hash mismatch for %s."), GetEntryName(It));
-				ErrorCount++;
-			}
-			else
-			{
-				UE_LOG(LogPakFile, Verbose, TEXT("%s OK. [%s]"), GetEntryName(It), *Info.GetCompressionMethod(EntryFromIndex.CompressionMethodIndex).ToString());
-			}
-			FMemory::Free(FileContents);
-		}
-		if (ErrorCount == 0)
-		{
-			UE_LOG(LogPakFile, Display, TEXT("Pak file \"%s\" healthy, %d files checked."), *PakFilename, FileCount);
-		}
-		else
-		{
-			UE_LOG(LogPakFile, Display, TEXT("Pak file \"%s\" corrupted (%d errors out of %d files checked.)."), *PakFilename, ErrorCount, FileCount);
-		}
-	}
-
-	double EndTime = FPlatformTime::Seconds();
-	double ElapsedTime = EndTime - StartTime;
-	UE_LOG(LogPakFile, Display, TEXT("Pak file \"%s\" checked in %.2fs"), *PakFilename, ElapsedTime);
-
-	return ErrorCount == 0;
-}
-
-bool CheckIoStoreContainerBlockSignatures(const TCHAR* InContainerPath)
-{
-	TRACE_CPUPROFILER_EVENT_SCOPE(CheckIoStoreContainerBlockSignatures);
-	UE_LOG(LogPakFile, Display, TEXT("Checking container file \"%s\"..."), InContainerPath);
-	double StartTime = FPlatformTime::Seconds();
-
-	FIoStoreTocResource TocResource;
-	FIoStatus Status = FIoStoreTocResource::Read(InContainerPath, EIoStoreTocReadOptions::Default, TocResource);
-	if (!Status.IsOk())
-	{
-		UE_LOG(LogPakFile, Error, TEXT("Failed reading toc file \"%s\"."), InContainerPath);
-		return false;
-	}
-
-	if (TocResource.ChunkBlockSignatures.Num() != TocResource.CompressionBlocks.Num())
-	{
-		UE_LOG(LogPakFile, Error, TEXT("Toc file \"%s\" doesn't contain any chunk block signatures."), InContainerPath);
-		return false;
-	}
-
-	TUniquePtr<FArchive> ContainerFileReader;
-	int32 LastPartitionIndex = -1;
-	IPlatformFile& Ipf = FPlatformFileManager::Get().GetPlatformFile();
-	TArray<uint8> BlockBuffer;
-	BlockBuffer.SetNum(static_cast<int32>(TocResource.Header.CompressionBlockSize));
-	const int32 BlockCount = TocResource.CompressionBlocks.Num();
-	int32 ErrorCount = 0;
-	FString ContainerBasePath = FPaths::ChangeExtension(InContainerPath, TEXT(""));
-	TStringBuilder<256> UcasFilePath;
-	for (int32 BlockIndex = 0; BlockIndex < BlockCount; ++BlockIndex)
-	{
-		const FIoStoreTocCompressedBlockEntry& CompressionBlockEntry = TocResource.CompressionBlocks[BlockIndex];
-		uint64 BlockRawSize = Align(CompressionBlockEntry.GetCompressedSize(), FAES::AESBlockSize);
-		check(BlockRawSize <= TocResource.Header.CompressionBlockSize);
-		const int32 PartitionIndex = int32(CompressionBlockEntry.GetOffset() / TocResource.Header.PartitionSize);
-		const uint64 PartitionRawOffset = CompressionBlockEntry.GetOffset() % TocResource.Header.PartitionSize;
-		if (PartitionIndex != LastPartitionIndex)
-		{
-			UcasFilePath.Reset();
-			UcasFilePath.Append(ContainerBasePath);
-			if (PartitionIndex > 0)
-			{
-				UcasFilePath.Append(FString::Printf(TEXT("_s%d"), PartitionIndex));
-			}
-			UcasFilePath.Append(TEXT(".ucas"));
-			IFileHandle* ContainerFileHandle = Ipf.OpenRead(*UcasFilePath, /* allowwrite */ false);
-			if (!ContainerFileHandle)
-			{
-				UE_LOG(LogPakFile, Error, TEXT("Failed opening container file \"%s\"."), *UcasFilePath);
-				return false;
-			}
-			ContainerFileReader.Reset(new FArchiveFileReaderGeneric(ContainerFileHandle, *UcasFilePath, ContainerFileHandle->Size(), 256 << 10));
-			LastPartitionIndex = PartitionIndex;
-		}
-		ContainerFileReader->Seek(PartitionRawOffset);
-		ContainerFileReader->Precache(PartitionRawOffset, 0); // Without this buffering won't work due to the first read after a seek always being uncached
-		ContainerFileReader->Serialize(BlockBuffer.GetData(), BlockRawSize);
-		FSHAHash BlockHash;
-		FSHA1::HashBuffer(BlockBuffer.GetData(), BlockRawSize, BlockHash.Hash);
-		if (TocResource.ChunkBlockSignatures[BlockIndex] != BlockHash)
-		{
-			UE_LOG(LogPakFile, Warning, TEXT("Hash mismatch for block [%i/%i]! Expected %s, Received %s"), BlockIndex, BlockCount, *TocResource.ChunkBlockSignatures[BlockIndex].ToString(), *BlockHash.ToString());
-
-			FPakChunkSignatureCheckFailedData Data(*UcasFilePath, TPakChunkHash(), TPakChunkHash(), BlockIndex);
-#if PAKHASH_USE_CRC
-			Data.ExpectedHash = GetTypeHash(TocResource.ChunkBlockSignatures[BlockIndex]);
-			Data.ReceivedHash = GetTypeHash(BlockHash);
-#else
-			Data.ExpectedHash = TocResource.ChunkBlockSignatures[BlockIndex];
-			Data.ReceivedHash = BlockHash;
-#endif
-			FPakPlatformFile::BroadcastPakChunkSignatureCheckFailure(Data);
-			++ErrorCount;
-		}
-	}
-
-	double EndTime = FPlatformTime::Seconds();
-	double ElapsedTime = EndTime - StartTime;
-	UE_LOG(LogPakFile, Display, TEXT("Container file \"%s\" checked in %.2fs"), InContainerPath, ElapsedTime);
-
-	return ErrorCount == 0;
-}
-
-#if DO_CHECK
-/**
-* FThreadCheckingArchiveProxy - checks that inner archive is only used from the specified thread ID
-*/
-class FThreadCheckingArchiveProxy : TUniquePtr<FArchive>, public FArchiveProxy
-{
-public:
-
-	const uint32 ThreadId;
-
-	FThreadCheckingArchiveProxy(FArchive* InReader, uint32 InThreadId)
-		: TUniquePtr(InReader) // Make sure proxy is destroyed before InReader
-		, FArchiveProxy(*InReader)
-		, ThreadId(InThreadId)
-	{}
-
-	//~ Begin FArchiveProxy Interface
-	virtual void Serialize(void* Data, int64 Length) override
-	{
-		if (FPlatformTLS::GetCurrentThreadId() != ThreadId)
-		{
-			UE_LOG(LogPakFile, Error, TEXT("Attempted serialize using thread-specific pak file reader on the wrong thread.  Reader for thread %d used by thread %d."), ThreadId, FPlatformTLS::GetCurrentThreadId());
-		}
-		InnerArchive.Serialize(Data, Length);
-	}
-
-	virtual void Seek(int64 InPos) override
-	{
-		if (FPlatformTLS::GetCurrentThreadId() != ThreadId)
-		{
-			UE_LOG(LogPakFile, Error, TEXT("Attempted seek using thread-specific pak file reader on the wrong thread.  Reader for thread %d used by thread %d."), ThreadId, FPlatformTLS::GetCurrentThreadId());
-		}
-		InnerArchive.Seek(InPos);
-	}
-	//~ End FArchiveProxy Interface
-};
-#endif //DO_CHECK
-
-void FPakFile::GetPrunedFilenames(TArray<FString>& OutFileList) const
-{
-	for (FFilenameIterator It(*this, true /* bIncludeDeleted */); It; ++It)
-	{
-		OutFileList.Add(PakPathCombine(MountPoint, It.Filename()));
-	}
-}
-
-void FPakFile::GetPrunedFilenamesInChunk(const TArray<int32>& InChunkIDs, TArray<FString>& OutFileList) const
-{
-	for (FFilenameIterator It(*this, true /* bIncludeDeleted */); It; ++It)
-	{
-		const FPakEntry& File = It.Info();
-		int64 FileStart = File.Offset;
-		int64 FileEnd = File.Offset + File.Size;
-
-		for (int64 LocalChunkID : InChunkIDs)
-		{
-			int64 ChunkStart = LocalChunkID * FPakInfo::MaxChunkDataSize;
-			int64 ChunkEnd = ChunkStart + FPakInfo::MaxChunkDataSize;
-
-			if (FileStart < ChunkEnd && FileEnd > ChunkStart)
-			{
-				OutFileList.Add(It.Filename());
-				break;
-			}
-		}
-	}
-}
-
-
-/**
- * Search the given FDirectoryIndex for all files under the given Directory.  Helper for FindFilesAtPath, called separately on the DirectoryIndex or Pruned DirectoryIndex. Does not use
- * FScopedPakDirectoryIndexAccess internally; caller is responsible for calling from within a lock.
- * Returned paths are full paths (include the mount point)
- */
-template <typename ShouldVisitFunc, class ContainerType>
-void FPakFile::FindFilesAtPathInIndex(const FDirectoryIndex& TargetIndex, ContainerType& OutFiles,
-	const FString& Directory, const ShouldVisitFunc& ShouldVisit, bool bIncludeFiles, bool bIncludeDirectories,
-	bool bRecursive) const
-{
-	// Early out if MountPoint is not matching directory
-	if (!Directory.StartsWith(MountPoint))
-	{
-		return;
-	}
-
-	FStringView RelativeSearch(FStringView(Directory).RightChop(MountPoint.Len()));
-
-	TArray<FString> DirectoriesInPak; // List of all unique directories at path
-	for (TMap<FString, FPakDirectory>::TConstIterator It(TargetIndex); It; ++It)
-	{
-		// Check if the file is under the specified path.
-		if (FStringView(It.Key()).StartsWith(RelativeSearch))
-		{
-			FString PakPath = PakPathCombine(MountPoint, It.Key());
-			if (bRecursive == true)
-			{
-				// Add everything
-				if (bIncludeFiles)
-				{
-					for (FPakDirectory::TConstIterator DirectoryIt(It.Value()); DirectoryIt; ++DirectoryIt)
-					{
-						const FString& FilePathUnderDirectory = DirectoryIt.Key();
-						if (ShouldVisit(FilePathUnderDirectory))
-						{
-							OutFiles.Add(PakPathCombine(PakPath, FilePathUnderDirectory));
-						}
-					}
-				}
-				if (bIncludeDirectories)
-				{
-					if (Directory != PakPath)
-					{
-						if (ShouldVisit(PakPath))
-						{
-							DirectoriesInPak.Add(MoveTemp(PakPath));
-						}
-					}
-				}
-			}
-			else
-			{
-				int32 SubDirIndex = PakPath.Len() > Directory.Len() ? PakPath.Find(TEXT("/"), ESearchCase::CaseSensitive, ESearchDir::FromStart, Directory.Len() + 1) : INDEX_NONE;
-				// Add files in the specified folder only.
-				if (bIncludeFiles && SubDirIndex == INDEX_NONE)
-				{
-					for (FPakDirectory::TConstIterator DirectoryIt(It.Value()); DirectoryIt; ++DirectoryIt)
-					{
-						const FString& FilePathUnderDirectory = DirectoryIt.Key();
-						if (ShouldVisit(FilePathUnderDirectory))
-						{
-							OutFiles.Add(PakPathCombine(PakPath, FilePathUnderDirectory));
-						}
-					}
-				}
-				// Add sub-folders in the specified folder only
-				if (bIncludeDirectories && SubDirIndex >= 0)
-				{
-					FString SubDirPath = PakPath.Left(SubDirIndex + 1);
-					if (ShouldVisit(SubDirPath))
-					{
-						DirectoriesInPak.AddUnique(MoveTemp(SubDirPath));
-					}
-				}
-			}
-		}
-	}
-	OutFiles.Append(MoveTemp(DirectoriesInPak));
-}
-
-template <typename ShouldVisitFunc, class ContainerType>
-void FPakFile::FindPrunedFilesAtPathInternal(const TCHAR* InPath, const ShouldVisitFunc& ShouldVisit, ContainerType& OutFiles,
-	bool bIncludeFiles, bool bIncludeDirectories, bool bRecursive) const
-{
-	// Make sure all directory names end with '/'.
-	FString Directory(InPath);
-	MakeDirectoryFromPath(Directory);
-
-	// Check the specified path is under the mount point of this pak file.
-	// The reverse case (MountPoint StartsWith Directory) is needed to properly handle
-	// pak files that are a subdirectory of the actual directory.
-	if (!Directory.StartsWith(MountPoint) && !MountPoint.StartsWith(Directory))
-	{
-		return;
-	}
-
-	FScopedPakDirectoryIndexAccess ScopeAccess(*this);
-#if ENABLE_PAKFILE_RUNTIME_PRUNING_VALIDATE
-	if (ShouldValidatePrunedDirectory())
-	{
-		TSet<FString> FullFoundFiles, PrunedFoundFiles;
-		FindFilesAtPathInIndex(DirectoryIndex, FullFoundFiles, Directory, ShouldVisit,
-			bIncludeFiles, bIncludeDirectories, bRecursive);
-		FindFilesAtPathInIndex(PrunedDirectoryIndex, PrunedFoundFiles, Directory, ShouldVisit,
-			bIncludeFiles, bIncludeDirectories, bRecursive);
-		ValidateDirectorySearch(FullFoundFiles, PrunedFoundFiles, InPath);
-
-		for (const FString& FoundFile : FullFoundFiles)
-		{
-			OutFiles.Add(FoundFile);
-		}
-	}
-	else
-#endif
-	{
-		FindFilesAtPathInIndex(DirectoryIndex, OutFiles, Directory, ShouldVisit,
-			bIncludeFiles, bIncludeDirectories, bRecursive);
-	}
-}
-
-void FPakFile::FindPrunedFilesAtPath(const TCHAR* InPath, TArray<FString>& OutFiles,
-	bool bIncludeFiles, bool bIncludeDirectories, bool bRecursive) const
-{
-	auto ShouldVisit = [](FStringView Path) { return true; };
-	FindPrunedFilesAtPathInternal(InPath, ShouldVisit, OutFiles, bIncludeFiles, bIncludeDirectories, bRecursive);
-}
-
-
-#if ENABLE_PAKFILE_RUNTIME_PRUNING_VALIDATE
-void FPakFile::ValidateDirectorySearch(const TSet<FString>& FullFoundFiles, const TSet<FString>& PrunedFoundFiles, const TCHAR* InPath) const
-{
-	TArray<FString> MissingFromPruned;
-	for (const FString& FileInFull : FullFoundFiles)
-	{
-		if (!PrunedFoundFiles.Contains(FileInFull))
-		{
-			MissingFromPruned.Add(FileInFull);
-		}
-	}
-	TArray<FString> MissingFromFull;
-	for (const FString& FileInPruned : PrunedFoundFiles)
-	{
-		if (!FullFoundFiles.Contains(FileInPruned))
-		{
-			MissingFromFull.Add(FileInPruned);
-		}
-	}
-
-	if (MissingFromPruned.Num() == 0 && MissingFromFull.Num() == 0)
-	{
-		return;
-	}
-
-	TArray<FString> WildCards, OldWildCards;
-	GConfig->GetArray(TEXT("Pak"), TEXT("IndexValidationIgnore"), WildCards, GEngineIni);
-	auto IsIgnore = [&WildCards](const FString& FilePath)
-	{
-		for (const FString& WildCard : WildCards)
-		{
-			if (FilePath.MatchesWildcard(WildCard))
-			{
-				return true;
-			}
-		}
-		return false;
-	};
-	auto StripIgnores = [&IsIgnore](TArray<FString>& FilePaths)
-	{
-		for (int Idx = FilePaths.Num() - 1; Idx >= 0; --Idx)
-		{
-			if (IsIgnore(FilePaths[Idx]))
-			{
-				FilePaths.RemoveAtSwap(Idx);
-			}
-		}
-	};
-
-	StripIgnores(MissingFromPruned);
-	StripIgnores(MissingFromFull);
-
-	if (MissingFromPruned.Num() == 0 && MissingFromFull.Num() == 0)
-	{
-		return;
-	}
-	MissingFromPruned.Sort();
-	MissingFromFull.Sort();
-
-	// TODO: Restore this as an Error once we modify IPlatformFile::IterateDirectoryRecursively to declare its filefilter so we can ignore the spurious
-	// discovered files that are not part of the fully filtered query
-	UE_LOG(LogPakFile, Error, TEXT("FindPrunedFilesAtPath('%s') for PakFile '%s' found a different list in the FullDirectory than in the PrunedDirectory. ")
-		TEXT("Change the calling code or add the files to Engine:[Pak]:WildcardsToKeepInPakStringIndex or Engine:[Pak]:IndexValidationIgnore."),
-		InPath, *PakFilename);
-
-#if !NO_LOGGING && !UE_BUILD_SHIPPING
-	// Logging callstacks is expensive (multiple seconds long). Only do it the first time a path is seen, and only for the first
-	// few paths.
-	static TSet<FString> AlreadyLoggedCallstack;
-	static FCriticalSection AlreadyLoggedCallstackLock;
-	constexpr int32 CallstackLogDirsMax = 10;
-	bool bShouldLogCallstack = false;
-	if (AlreadyLoggedCallstack.Num() < CallstackLogDirsMax) // check to avoid taking critical section if unnecessary
-	{
-		FScopeLock AlreadyLoggedCallstackScopeLock(&AlreadyLoggedCallstackLock);
-		if (AlreadyLoggedCallstack.Num() < CallstackLogDirsMax) // check again since other thread may have modified it
-		{
-			bool bAlreadyLogged;
-			AlreadyLoggedCallstack.Add(FString(InPath), &bAlreadyLogged);
-			bShouldLogCallstack = !bAlreadyLogged;
-		}
-	}
-	if (bShouldLogCallstack)
-	{
-		UE_LOG(LogPakFile, Warning, TEXT("Callstack of FindPrunedFilesAtPath('%s'):"), InPath);
-		FDebug::DumpStackTraceToLog(ELogVerbosity::Warning);
-	}
-#endif
-
-	if (MissingFromPruned.Num() > 0)
-	{
-		for (const FString& Missing : MissingFromPruned)
-		{
-			UE_LOG(LogPakFile, Warning, TEXT("MissingPrunedPakFile: %s"), *Missing);
-		}
-	}
-	if (MissingFromFull.Num() > 0)
-	{
-		UE_LOG(LogPakFile, Error, TEXT("Some files in the PrunedDirectory are missing from the FullDirectory.  This is a logic error in FPakFile since the PrunedDirectory should be a subset of the FullDirectory."));
-		for (const FString& Missing : MissingFromFull)
-		{
-			UE_LOG(LogPakFile, Warning, TEXT("MissingFullPakFile: %s"), *Missing);
-		}
-	}
-}
-#endif
-
-bool FPakFile::RecreatePakReaders(IPlatformFile* LowerLevel)
-{
-	FScopeLock ScopedLock(&ReadersCriticalSection);
-
-	if (CurrentlyUsedReaders > 0)
-	{
-		UE_LOG(LogPakFile, Error, TEXT("Recreating pak readers while we have readers loaned out, this may be lead to crashes or decryption problems."));
-	}
-
-	// need to reset the decryptor as it will hold a pointer to the first created pak reader
-	Decryptor.Reset();
-
-	TArray<FArchiveAndLastAccessTime> TempReaders;
-
-	// Create a new PakReader *per* instance that was already mapped
-	for (const FArchiveAndLastAccessTime& Reader : Readers)
-	{
-		TUniquePtr<FArchive> PakReader = TUniquePtr<FArchive>(CreatePakReader(LowerLevel, *GetFilename()));
-		if (!PakReader)
-		{
-			UE_LOG(LogPakFile, Warning, TEXT("Unable to re-create pak \"%s\" handle"), *GetFilename());
-			return false;
-		}
-		TempReaders.Add(FArchiveAndLastAccessTime{ MoveTemp(PakReader), Reader.LastAccessTime });
-	}
-
-	// replace the current Readers with the newly created pak readers leaving them to out of scope
-	Readers= MoveTemp(TempReaders);
-
-	return true;
-}
-
-LLM_DEFINE_TAG(PakSharedReaders);
-
-FSharedPakReader FPakFile::GetSharedReader(IPlatformFile* LowerLevel)
-{
-	LLM_SCOPE_BYTAG(PakSharedReaders);
-	FArchive* PakReader = nullptr;
-	{
-		FScopeLock ScopedLock(&ReadersCriticalSection);
-		if (Readers.Num())
-		{
-			FArchiveAndLastAccessTime Reader = Readers.Pop();
-			PakReader = Reader.Archive.Release();
-		}
-		else
-		{
-			// Create a new FArchive reader and pass it to the new handle.
-			PakReader = CreatePakReader(LowerLevel, *GetFilename());
-
-			if (!PakReader)
-			{
-				UE_LOG(LogPakFile, Warning, TEXT("Unable to create pak \"%s\" handle"), *GetFilename());
-			}
-		}
-		++CurrentlyUsedReaders;
-	}
-
-	return FSharedPakReader(PakReader, this);
-}
-
-void FPakFile::ReturnSharedReader(FArchive* Archive)
-{
-	FScopeLock ScopedLock(&ReadersCriticalSection);
-	--CurrentlyUsedReaders;
-	Readers.Push(FArchiveAndLastAccessTime{ TUniquePtr<FArchive>{Archive }, FPlatformTime::Seconds()});
-}
-
-void FPakFile::ReleaseOldReaders(double MaxAgeSeconds)
-{
-	if (ReadersCriticalSection.TryLock())
-	{
-		ON_SCOPE_EXIT
-		{
-			ReadersCriticalSection.Unlock();
-		};
-		double SearchTime = FPlatformTime::Seconds() - MaxAgeSeconds;
-		for (int32 i = Readers.Num() - 1; i >= 0; --i)
-		{
-			const FArchiveAndLastAccessTime& Reader = Readers[i];
-			if (Reader.LastAccessTime <= SearchTime)
-			{
-				// Remove this and all readers older than it (pushed before it)
-				Readers.RemoveAt(0, i + 1);
-				break;
-			}
-		}
-
-		if (Readers.Num() == 0 && CurrentlyUsedReaders == 0)
-		{
-			Decryptor.Reset();
-		}
-	}
-}
-
-const FPakEntryLocation* FPakFile::FindLocationFromIndex(const FString& FullPath, const FString& MountPoint, const FPathHashIndex& PathHashIndex, uint64 PathHashSeed, int32 PakFileVersion)
-{
-	const TCHAR* RelativePathFromMount = GetRelativeFilePathFromMountPointer(FullPath, MountPoint);
-	if (!RelativePathFromMount)
-	{
-		return nullptr;
-	}
-	uint64 PathHash = HashPath(RelativePathFromMount, PathHashSeed, PakFileVersion);
-	return PathHashIndex.Find(PathHash);
-}
-
-const FPakEntryLocation* FPakFile::FindLocationFromIndex(const FString& FullPath, const FString& MountPoint, const FDirectoryIndex& DirectoryIndex)
-{
-	if (!FullPath.StartsWith(MountPoint))
-	{
-		return nullptr;
-	}
-	const TCHAR* RelativePathFromMount = (*FullPath) + MountPoint.Len();
-	FString RelativeDirName(RelativePathFromMount);
-	FString CleanFileName;
-	if (RelativeDirName.IsEmpty())
-	{
-		return nullptr;
-	}
-	SplitPathInline(RelativeDirName, CleanFileName);
-	const FPakDirectory* PakDirectory = DirectoryIndex.Find(RelativeDirName);
-	if (PakDirectory)
-	{
-		return PakDirectory->Find(CleanFileName);
-	}
-	return nullptr;
-}
-
-FPakFile::EFindResult FPakFile::Find(const FString& FullPath, FPakEntry* OutEntry) const
-{
-	//QUICK_SCOPE_CYCLE_COUNTER(PakFileFind);
-
-	const FPakEntryLocation* PakEntryLocation;
-#if ENABLE_PAKFILE_RUNTIME_PRUNING_VALIDATE
-	if (IsPakValidatePruning() && bHasPathHashIndex && bHasFullDirectoryIndex)
-	{
-		const FPakEntryLocation* PathHashLocation = nullptr;
-		PathHashLocation = FindLocationFromIndex(FullPath, MountPoint, PathHashIndex, PathHashSeed, Info.Version);
-
-		const FPakEntryLocation* DirectoryLocation = nullptr;
-
-		{
-			FScopedPakDirectoryIndexAccess ScopeAccess(*this);
-			DirectoryLocation = FindLocationFromIndex(FullPath, MountPoint, DirectoryIndex);
-		}
-
-		if ((PathHashLocation != nullptr) != (DirectoryLocation != nullptr))
-		{
-			const TCHAR* FoundName = TEXT("PathHashIndex");
-			const TCHAR* NotFoundName = TEXT("FullDirectoryIndex");
-			if (!PathHashLocation)
-			{
-				Swap(FoundName, NotFoundName);
-			}
-			UE_LOG(LogPakFile, Error, TEXT("PathHashIndex does not match FullDirectoryIndex. Pakfile '%s' has '%s' in its %s but not in its %s."),
-				*PakFilename, *FullPath, FoundName, NotFoundName);
-		}
-		PakEntryLocation = PathHashLocation ? PathHashLocation : DirectoryLocation;
-	}
-	else
-#endif
-	{
-		if (bHasPathHashIndex)
-		{
-			PakEntryLocation = FindLocationFromIndex(FullPath, MountPoint, PathHashIndex, PathHashSeed, Info.Version);
-		}
-		else
-		{
-			check(bHasFullDirectoryIndex);
-			FScopedPakDirectoryIndexAccess ScopeAccess(*this);
-			PakEntryLocation = FindLocationFromIndex(FullPath, MountPoint, DirectoryIndex);
-		}
-	}
-	if (!PakEntryLocation)
-	{
-		return EFindResult::NotFound;
-	}
-
-	return GetPakEntry(*PakEntryLocation, OutEntry);
-}
-
-#if !UE_BUILD_SHIPPING
+
+#if !UE_BUILD_SHIPPING
 class FPakExec : private FSelfRegisteringExec
@@ -8346,10 +5665,6 @@
 #endif
 }
 
-#if ENABLE_PAKFILE_RUNTIME_PRUNING
-bool FPakFile::bSomePakNeedsPruning = false;
-#endif
-
 #if !UE_BUILD_SHIPPING
 uint64 GetRecursiveAllocatedSize(const FPakDirectory& Index)
 {
@@ -9144,73 +6459,6 @@
 	return Result;
 }
 
-/**
- * Module for the pak file
- */
-class FPakFileModule : public IPlatformFileModule
-{
-public:
-	virtual IPlatformFile* GetPlatformFile() override
-	{
-		check(Singleton.IsValid());
-		return Singleton.Get();
-	}
-
-	virtual void StartupModule() override
-	{
-		Singleton = MakeUnique<FPakPlatformFile>();
-		FModuleManager::LoadModuleChecked<IModuleInterface>(TEXT("RSA"));
-	}
-
-	virtual void ShutdownModule() override
-	{
-		// remove ourselves from the platform file chain (there can be late writes after the shutdown).
-		if (Singleton.IsValid())
-		{
-			if (FPlatformFileManager::Get().FindPlatformFile(Singleton.Get()->GetName()))
-			{
-				FPlatformFileManager::Get().RemovePlatformFile(Singleton.Get());
-			}
-		}
-
-		Singleton.Reset();
-	}
-
-	TUniquePtr<IPlatformFile> Singleton;
-};
-
-void FPakFile::AddSpecialFile(const FPakEntry& Entry, const FString& Filename)
-{
-	MakeDirectoryFromPath(MountPoint);
-
-	// TODO: This function is not threadsafe; readers of the Indexes will be invalidated when we modify them
-	// To make it threadsafe would require always holding the lock around any read of either index, which is
-	// more expensive than we want to support this debug feature
-	FPakEntryLocation EntryLocation;
-	if (!Entry.IsDeleteRecord())
-	{
-		// Add new file info.
-		TArray<uint8> NewEncodedPakEntries;
-		FMemoryWriter MemoryWriter(NewEncodedPakEntries);
-		EntryLocation = FPakEntryLocation::CreateFromOffsetIntoEncoded(EncodedPakEntries.Num());
-		if (EncodePakEntry(MemoryWriter, Entry, Info))
-		{
-			EncodedPakEntries.Append(NewEncodedPakEntries);
-			EncodedPakEntries.Shrink();
-		}
-		else
-		{
-			EntryLocation = FPakEntryLocation::CreateFromListIndex(Files.Num());
-			Files.Add(Entry);
-			Files.Shrink();
-		}
-		NumEntries++;
-	}
-
-	FPathHashIndex* PathHashToWrite = bHasPathHashIndex ? &PathHashIndex : nullptr;
-	AddEntryToIndex(Filename, EntryLocation, MountPoint, PathHashSeed, &DirectoryIndex, PathHashToWrite, nullptr /* CollisionDetection */, Info.Version);
-}
-
 void FPakPlatformFile::MakeUniquePakFilesForTheseFiles(const TArray<TArray<FString>>& InFiles)
 {
 	for (int k = 0; k < InFiles.Num(); k++)
@@ -9261,8 +6509,6 @@
 	}
 }
 
-IMPLEMENT_MODULE(FPakFileModule, PakFile);
-
 #if !UE_BUILD_SHIPPING
 static void AsyncFileTest(const TArray<FString>& Args)
 {
diff --git a/Engine/Source/Runtime/PakFile/Private/PakFile.cpp b/Engine/Source/Runtime/PakFile/Private/PakFile.cpp
new file mode 100644
--- /dev/null
+++ b/Engine/Source/Runtime/PakFile/Private/PakFile.cpp
@@ -0,0 +1,1943 @@
+// Copyright Epic Games, Inc. All Rights Reserved.
+
+#include "Async/MappedFileHandle.h"
+#include "HAL/FileManager.h"
+#include "HAL/FileManagerGeneric.h"
+#include "HAL/LowLevelMemTracker.h"
+#include "IO/IoContainerHeader.h"
+#include "IPlatformFilePak.h"
+#include "Math/GuardedInt.h"
+#include "Misc/CommandLine.h"
+#include "Misc/ConfigCacheIni.h"
+#include "Misc/EncryptionKeyManager.h"
+#include "Misc/Fnv.h"
+#include "Misc/Parse.h"
+#include "Misc/ScopeExit.h"
+#include "Serialization/MemoryReader.h"
+#include "Serialization/MemoryWriter.h"
+#include "SignedArchiveReader.h"
+
+#include "PakFile.inl"
+
+LLM_DEFINE_TAG(PakSharedReaders);
+
+int32	GetPakchunkIndexFromPakFile(const FString&);
+bool	ShouldCheckPak();
+void	DecryptData(uint8*, uint64, FGuid);
+
+#if ENABLE_PAKFILE_RUNTIME_PRUNING
+bool FPakFile::bSomePakNeedsPruning = false;
+#endif
+
+/*** This is a copy of FFnv::MemFnv64 from before the bugfix for swapped Offset
+ * and Prime. It is used to decode legacy paks that have hashes created from
+ * the prebugfix version of the function */
+static uint64 LegacyMemFnv64(const void* InData, int32 Length, uint64 InOffset)
+{
+	// constants from above reference
+	static const uint64 Offset = 0x00000100000001b3;
+	static const uint64 Prime = 0xcbf29ce484222325;
+
+	const uint8* __restrict Data = (uint8*)InData;
+
+	uint64 Fnv = Offset + InOffset; // this is not strictly correct as the offset should be prime and InOffset could be arbitrary
+	for (; Length; --Length)
+	{
+		Fnv ^= *Data++;
+		Fnv *= Prime;
+	}
+
+	return Fnv;
+}
+
+
+
+struct FPakFile::FIndexSettings
+{
+	FIndexSettings()
+	{
+		bKeepFullDirectory = true;
+		bValidatePruning = false;
+		bDelayPruning = false;
+		bWritePathHashIndex = true;
+		bWriteFullDirectoryIndex = true;
+
+		// Paks are mounted before config files are read, so the licensee needs to hardcode all settings used for runtime index loading rather than specifying them in ini
+		if (FPakPlatformFile::GetPakSetIndexSettingsDelegate().IsBound())
+		{
+			FPakPlatformFile::GetPakSetIndexSettingsDelegate().Execute(bKeepFullDirectory, bValidatePruning, bDelayPruning);
+		}
+
+		// Settings not used at runtime can be read from ini
+#if !UE_BUILD_SHIPPING
+		GConfig->GetBool(TEXT("Pak"), TEXT("WritePathHashIndex"), bWritePathHashIndex, GEngineIni);
+		GConfig->GetBool(TEXT("Pak"), TEXT("WriteFullDirectoryIndex"), bWriteFullDirectoryIndex, GEngineIni);
+#endif
+
+#if IS_PROGRAM || WITH_EDITOR
+		// Directory pruning is not enabled in the editor or in development programs because there is no need to save the memory in those environments and some development features require not pruning
+		bKeepFullDirectory = true;
+#else
+		bKeepFullDirectory = bKeepFullDirectory || !FPlatformProperties::RequiresCookedData();
+#endif
+#if !UE_BUILD_SHIPPING
+		const TCHAR* CommandLine = FCommandLine::Get();
+		FParse::Bool(CommandLine, TEXT("ForcePakKeepFullDirectory="), bKeepFullDirectory);
+#if ENABLE_PAKFILE_RUNTIME_PRUNING_VALIDATE
+		FParse::Bool(CommandLine, TEXT("ForcePakValidatePruning="), bValidatePruning);
+		FParse::Bool(CommandLine, TEXT("ForcePakDelayPruning="), bDelayPruning);
+#endif
+		FParse::Bool(CommandLine, TEXT("ForcePakWritePathHashIndex="), bWritePathHashIndex);
+		FParse::Bool(CommandLine, TEXT("ForcePakWriteFullDirectoryIndex="), bWriteFullDirectoryIndex);
+#endif
+	}
+
+	bool bKeepFullDirectory;
+	bool bValidatePruning;
+	bool bDelayPruning;
+	bool bWritePathHashIndex;
+	bool bWriteFullDirectoryIndex;
+};
+
+
+
+#if IS_PROGRAM
+FPakFile::FPakFile(const TCHAR* Filename, bool bIsSigned)
+	: PakFilename(Filename)
+	, PakFilenameName(Filename)
+	, PathHashSeed(0)
+	, NumEntries(0)
+	, CachedTotalSize(0)
+	, bSigned(bIsSigned)
+	, bIsValid(false)
+	, bHasPathHashIndex(false)
+	, bHasFullDirectoryIndex(false)
+#if ENABLE_PAKFILE_RUNTIME_PRUNING
+	, bWillPruneDirectoryIndex(false)
+	, bNeedsLegacyPruning(false)
+#endif
+	, PakchunkIndex(GetPakchunkIndexFromPakFile(Filename))
+ 	, MappedFileHandle(nullptr)
+	, CacheType(FPakFile::ECacheType::Shared)
+	, CacheIndex(-1)
+	, UnderlyingCacheTrimDisabled(false)
+	, bIsMounted(false)
+{
+	FSharedPakReader Reader = GetSharedReader(NULL);
+	if (Reader)
+	{
+		Timestamp = IFileManager::Get().GetTimeStamp(Filename);
+		Initialize(Reader.GetArchive());
+	}
+}
+#endif
+
+FPakFile::FPakFile(IPlatformFile* LowerLevel, const TCHAR* Filename, bool bIsSigned, bool bLoadIndex)
+	: PakFilename(Filename)
+	, PakFilenameName(Filename)
+	, PathHashSeed(0)
+	, NumEntries(0)
+	, CachedTotalSize(0)
+	, bSigned(bIsSigned)
+	, bIsValid(false)
+	, bHasPathHashIndex(false)
+	, bHasFullDirectoryIndex(false)
+#if ENABLE_PAKFILE_RUNTIME_PRUNING
+	, bWillPruneDirectoryIndex(false)
+	, bNeedsLegacyPruning(false)
+#endif
+	, PakchunkIndex(GetPakchunkIndexFromPakFile(Filename))
+	, MappedFileHandle(nullptr)
+	, CacheType(FPakFile::ECacheType::Shared)
+	, CacheIndex(-1)
+	, UnderlyingCacheTrimDisabled(false)
+	, bIsMounted(false)
+{
+	FSharedPakReader Reader = GetSharedReader(LowerLevel);
+	if (Reader)
+	{
+		Timestamp = LowerLevel->GetTimeStamp(Filename);
+		Initialize(Reader.GetArchive(), bLoadIndex);
+	}
+}
+
+#if WITH_EDITOR
+FPakFile::FPakFile(FArchive* Archive)
+	: PathHashSeed(0)
+	, NumEntries(0)
+	, bSigned(false)
+	, bIsValid(false)
+	, bHasPathHashIndex(false)
+	, bHasFullDirectoryIndex(false)
+#if ENABLE_PAKFILE_RUNTIME_PRUNING
+	, bWillPruneDirectoryIndex(false)
+	, bNeedsLegacyPruning(false)
+#endif
+	, PakchunkIndex(INDEX_NONE)
+	, MappedFileHandle(nullptr)
+	, CacheType(FPakFile::ECacheType::Shared)
+	, CacheIndex(-1)
+	, UnderlyingCacheTrimDisabled(false)
+	, bIsMounted(false)
+{
+	Initialize(*Archive);
+}
+#endif
+
+FPakFile::~FPakFile()
+{
+	delete MappedFileHandle;
+}
+
+bool FPakFile::PassedSignatureChecks() const
+{
+	return Decryptor.IsValid() && Decryptor->IsValid();
+}
+
+FArchive* FPakFile::CreatePakReader(IPlatformFile* LowerLevel, const TCHAR* Filename)
+{
+	auto MakeArchive = [&]() -> FArchive* { 
+		if (LowerLevel)
+		{
+			if( IFileHandle* Handle = LowerLevel->OpenRead(Filename) )
+			{
+				return new FArchiveFileReaderGeneric(Handle, Filename, Handle->Size());
+			}
+			else
+			{
+				return nullptr;
+			}
+		}
+		else
+		{
+			return IFileManager::Get().CreateFileReader(Filename); 
+		}
+	};
+
+	bool bNeedsDecryptor = false;
+	if (FPlatformProperties::RequiresCookedData())
+	{
+		bool bShouldCheckSignature = bSigned || FParse::Param(FCommandLine::Get(), TEXT("signedpak")) || FParse::Param(FCommandLine::Get(), TEXT("signed"));
+#if !UE_BUILD_SHIPPING
+		bShouldCheckSignature &= !FParse::Param(FCommandLine::Get(), TEXT("FileOpenLog"));
+#endif
+		if (bShouldCheckSignature)
+		{
+			bNeedsDecryptor = true;
+		}			
+	}
+
+	if(bNeedsDecryptor && !Decryptor.IsValid())
+	{
+		TUniquePtr<FArchive> DecryptorReader{ MakeArchive() };
+		if (DecryptorReader.IsValid())
+		{
+			Decryptor = MakeUnique<FChunkCacheWorker>(MoveTemp(DecryptorReader), Filename);
+		}
+
+		if (!Decryptor.IsValid() || !Decryptor->IsValid())
+		{
+			return nullptr;
+		}
+	}
+
+	// Now we either have a Decryptor or we don't need it
+	check(!bNeedsDecryptor || Decryptor.IsValid());
+
+	TUniquePtr<FArchive> Archive{ MakeArchive() };
+	if (!Archive.IsValid())
+	{
+		return nullptr;
+	}
+
+	if (Decryptor.IsValid())
+	{
+		return new FSignedArchiveReader(Archive.Release(), Decryptor.Get());
+	}
+	else
+	{
+		return Archive.Release();
+	}
+}
+
+void FPakFile::Initialize(FArchive& Reader, bool bLoadIndex)
+{
+	CachedTotalSize = Reader.TotalSize();
+	bool bShouldLoad = false;
+	int32 CompatibleVersion = FPakInfo::PakFile_Version_Latest;
+
+	LLM_SCOPE(ELLMTag::FileSystem);
+
+	// Serialize trailer and check if everything is as expected.
+	// start up one to offset the -- below
+	CompatibleVersion++;
+	int64 FileInfoPos = -1;
+	do
+	{
+		// try the next version down
+		CompatibleVersion--;
+
+		FileInfoPos = CachedTotalSize - Info.GetSerializedSize(CompatibleVersion);
+		if (FileInfoPos >= 0)
+		{
+			Reader.Seek(FileInfoPos);
+			Reader.Precache(FileInfoPos, 0); // Inform the archive that we're going to repeatedly serialize from the current location
+
+			SCOPED_BOOT_TIMING("PakFile_SerilizeTrailer");
+
+			// Serialize trailer and check if everything is as expected.
+			Info.Serialize(Reader, CompatibleVersion);
+			if (Info.Magic == FPakInfo::PakFile_Magic)
+			{
+				bShouldLoad = true;
+			}
+		}
+	} while (!bShouldLoad && CompatibleVersion >= FPakInfo::PakFile_Version_Initial);
+
+	if (bShouldLoad)
+	{
+		UE_CLOG(Info.Magic != FPakInfo::PakFile_Magic, LogPakFile, Fatal, TEXT("Trailing magic number (%ud) in '%s' is different than the expected one. Verify your installation."), Info.Magic, *PakFilename);
+		UE_CLOG(!(Info.Version >= FPakInfo::PakFile_Version_Initial && Info.Version <= CompatibleVersion), LogPakFile, Fatal, TEXT("Invalid pak file version (%d) in '%s'. Verify your installation."), Info.Version, *PakFilename);
+		UE_CLOG(!(Info.IndexOffset >= 0 && Info.IndexOffset < CachedTotalSize), LogPakFile, Fatal, TEXT("Index offset for pak file '%s' is invalid (%lld is bigger than file size %lld)"), *PakFilename, Info.IndexOffset, CachedTotalSize);
+		UE_CLOG(!((Info.IndexOffset + Info.IndexSize) >= 0 && (Info.IndexOffset + Info.IndexSize) <= CachedTotalSize), LogPakFile, Fatal, TEXT("Index end offset for pak file '%s' is invalid (%lld)"), *PakFilename, Info.IndexOffset + Info.IndexSize);
+
+		// If we aren't using a dynamic encryption key, process the pak file using the embedded key
+		if (!Info.EncryptionKeyGuid.IsValid() || UE::FEncryptionKeyManager::Get().ContainsKey(Info.EncryptionKeyGuid))
+		{
+			if (bLoadIndex)
+			{
+				LoadIndex(Reader);
+			}
+
+			if (ShouldCheckPak())
+			{
+				ensure(Check());
+			}
+		}
+
+		if (Decryptor.IsValid())
+		{
+			TSharedPtr<const FPakSignatureFile, ESPMode::ThreadSafe> SignatureFile = Decryptor->GetSignatures();
+			if (SignatureFile->SignatureData.Num() == UE_ARRAY_COUNT(FSHAHash::Hash))
+			{
+				bIsValid = (FMemory::Memcmp(SignatureFile->SignatureData.GetData(), Info.IndexHash.Hash, SignatureFile->SignatureData.Num()) == 0);
+			}
+			else
+			{
+				bIsValid = false;
+			}
+		}
+		else
+		{
+			// LoadIndex should crash in case of an error, so just assume everything is ok if we got here.
+			bIsValid = true;
+		}
+	}
+}
+
+void FPakFile::LoadIndex(FArchive& Reader)
+{
+	if (Info.Version >= FPakInfo::PakFile_Version_PathHashIndex)
+	{
+		if (!LoadIndexInternal(Reader))
+		{
+			// Index loading failed. Try again
+			if (!LoadIndexInternal(Reader))
+			{
+				UE_LOG(LogPakFile, Fatal, TEXT("Corrupt pak index detected on pak file: %s"), *PakFilename);
+			}
+		}
+	}
+	else
+	{
+		SCOPED_BOOT_TIMING("PakFile_LoadLegacy");
+		if (!LoadLegacyIndex(Reader))
+		{
+			// Index loading failed. Try again
+			if (!LoadLegacyIndex(Reader))
+			{
+				UE_LOG(LogPakFile, Fatal, TEXT("Corrupt pak index detected on pak file: %s"), *PakFilename);
+			}
+		}
+	}
+}
+
+bool FPakFile::LoadIndexInternal(FArchive& Reader)
+{
+	bHasPathHashIndex = false;
+	bHasFullDirectoryIndex = false;
+#if ENABLE_PAKFILE_RUNTIME_PRUNING
+	bNeedsLegacyPruning = false;
+	bWillPruneDirectoryIndex = false;
+#endif
+
+	FGuardedInt64 IndexEndPosition = FGuardedInt64(Info.IndexOffset) + Info.IndexSize;
+	if (Info.IndexOffset < 0 || 
+		Info.IndexSize < 0 ||
+		IndexEndPosition.InvalidOrGreaterThan(CachedTotalSize) ||
+		IntFitsIn<int32>(Info.IndexSize) == false)
+	{
+		UE_LOG(LogPakFile, Fatal, TEXT("Corrupted index offset in pak file."));
+		return false;
+	}
+
+	TArray<uint8> PrimaryIndexData;
+	Reader.Seek(Info.IndexOffset);
+	PrimaryIndexData.SetNum((int32)(Info.IndexSize));
+	{
+		SCOPED_BOOT_TIMING("PakFile_LoadPrimaryIndex");
+		Reader.Serialize(PrimaryIndexData.GetData(), Info.IndexSize);
+	}
+
+	FSHAHash ComputedHash;
+	{
+		SCOPED_BOOT_TIMING("PakFile_HashPrimaryIndex");
+		if (!DecryptAndValidateIndex(Reader, PrimaryIndexData, Info.IndexHash, ComputedHash))
+		{
+			UE_LOG(LogPakFile, Log, TEXT("Corrupt pak PrimaryIndex detected!"));
+			UE_LOG(LogPakFile, Log, TEXT(" Filename: %s"), *PakFilename);
+			UE_LOG(LogPakFile, Log, TEXT(" Encrypted: %d"), Info.bEncryptedIndex);
+			UE_LOG(LogPakFile, Log, TEXT(" Total Size: %lld"), Reader.TotalSize());
+			UE_LOG(LogPakFile, Log, TEXT(" Index Offset: %lld"), Info.IndexOffset);
+			UE_LOG(LogPakFile, Log, TEXT(" Index Size: %lld"), Info.IndexSize);
+			UE_LOG(LogPakFile, Log, TEXT(" Stored Index Hash: %s"), *Info.IndexHash.ToString());
+			UE_LOG(LogPakFile, Log, TEXT(" Computed Index Hash: %s"), *ComputedHash.ToString());
+			return false;
+		}
+	}
+
+	FMemoryReader PrimaryIndexReader(PrimaryIndexData);
+
+	// Read the scalar data (mount point, numentries, etc) and all entries.
+	NumEntries = 0;
+	PrimaryIndexReader << MountPoint;
+	// We are just deserializing a string, which could get however long. Since we know it's a path
+	// and paths are bound to file system rules, we know it can't get absurdly long (e.g. windows is _at best_ 32k)
+	// we just sanity check it to prevent operating on massive buffers and risking overflows.
+	if (MountPoint.Len() > 65535)
+	{
+		UE_LOG(LogPakFile, Error, TEXT("Corrupt pak index data: MountPoint path is longer than 65k"));
+		return false;
+	}
+
+	MakeDirectoryFromPath(MountPoint);
+	PrimaryIndexReader << NumEntries;
+	if (NumEntries < 0)
+	{
+		UE_LOG(LogPakFile, Error, TEXT("Corrupt pak index data: Negative entries count in pak file."));
+		return false;
+	}
+	PrimaryIndexReader << PathHashSeed;
+
+	bool bReaderHasPathHashIndex = false;
+	int64 PathHashIndexOffset = INDEX_NONE;
+	int64 PathHashIndexSize = 0;
+	FSHAHash PathHashIndexHash;
+	PrimaryIndexReader << bReaderHasPathHashIndex;
+	if (bReaderHasPathHashIndex)
+	{
+		PrimaryIndexReader << PathHashIndexOffset;
+		PrimaryIndexReader << PathHashIndexSize;
+		PrimaryIndexReader << PathHashIndexHash;
+		bReaderHasPathHashIndex = bReaderHasPathHashIndex && PathHashIndexOffset != INDEX_NONE;
+	}
+
+	bool bReaderHasFullDirectoryIndex = false;
+	int64 FullDirectoryIndexOffset = INDEX_NONE;
+	int64 FullDirectoryIndexSize = 0;
+	FSHAHash FullDirectoryIndexHash;
+	PrimaryIndexReader << bReaderHasFullDirectoryIndex;
+	if (bReaderHasFullDirectoryIndex)
+	{
+		PrimaryIndexReader << FullDirectoryIndexOffset;
+		PrimaryIndexReader << FullDirectoryIndexSize;
+		PrimaryIndexReader << FullDirectoryIndexHash;
+		bReaderHasFullDirectoryIndex = bReaderHasFullDirectoryIndex && FullDirectoryIndexOffset  != INDEX_NONE;
+	}
+	{
+		SCOPED_BOOT_TIMING("PakFile_SerializeEncodedEntries");
+		PrimaryIndexReader << EncodedPakEntries;
+	}
+
+	int32 FilesNum = 0;
+	PrimaryIndexReader << FilesNum;
+	if (FilesNum < 0)
+	{
+		// Should not be possible for any values in the PrimaryIndex to be invalid, since we verified the index hash
+		UE_LOG(LogPakFile, Log, TEXT("Corrupt pak PrimaryIndex detected!"));
+		UE_LOG(LogPakFile, Log, TEXT(" FilesNum: %d"), FilesNum);
+		return false;
+	}
+	Files.SetNum(FilesNum);
+	if (FilesNum > 0)
+	{
+		SCOPED_BOOT_TIMING("PakFile_SerializeUnencodedEntries");
+		FPakEntry* FileEntries = Files.GetData();
+		for (int32 FileIndex = 0; FileIndex < FilesNum; ++FileIndex)
+		{
+			FileEntries[FileIndex].Serialize(PrimaryIndexReader, Info.Version);
+		}
+	}
+
+	// Decide which SecondaryIndex(es) to load
+	bool bWillUseFullDirectoryIndex;
+	bool bWillUsePathHashIndex;
+	bool bReadFullDirectoryIndex;
+	if (bReaderHasPathHashIndex && bReaderHasFullDirectoryIndex)
+	{
+		bWillUseFullDirectoryIndex = IsPakKeepFullDirectory();
+		bWillUsePathHashIndex = !bWillUseFullDirectoryIndex;
+#if ENABLE_PAKFILE_RUNTIME_PRUNING
+		bool bWantToReadFullDirectoryIndex = IsPakKeepFullDirectory() || IsPakValidatePruning() || IsPakDelayPruning();
+#else
+		bool bWantToReadFullDirectoryIndex = IsPakKeepFullDirectory();
+#endif
+		bReadFullDirectoryIndex = bReaderHasFullDirectoryIndex && bWantToReadFullDirectoryIndex;
+	}
+	else if (bReaderHasPathHashIndex)
+	{
+		bWillUsePathHashIndex = true;
+		bWillUseFullDirectoryIndex = false;
+		bReadFullDirectoryIndex = false;
+	}
+	else if (bReaderHasFullDirectoryIndex)
+	{
+		// We don't support creating the PathHash Index at runtime; we want to move to having only the PathHashIndex, so supporting not having it at all is not useful enough to write
+		bWillUsePathHashIndex = false;
+		bWillUseFullDirectoryIndex = true;
+		bReadFullDirectoryIndex = true;
+	}
+	else
+	{
+		// It should not be possible for PrimaryIndexes to be built without a PathHashIndex AND without a FullDirectoryIndex; CreatePakFile in UnrealPak.exe has a check statement for it.
+		UE_LOG(LogPakFile, Log, TEXT("Corrupt pak PrimaryIndex detected!"));
+		UE_LOG(LogPakFile, Log, TEXT(" bReaderHasPathHashIndex: false"));
+		UE_LOG(LogPakFile, Log, TEXT(" bReaderHasFullDirectoryIndex: false"));
+		return false;
+	}
+
+	// Load the Secondary Index(es)
+	TArray<uint8> PathHashIndexData;
+	FMemoryReader PathHashIndexReader(PathHashIndexData);
+	if (bWillUsePathHashIndex)
+	{
+		FGuardedInt64 PathHashIndexEndPosition = FGuardedInt64(PathHashIndexOffset) + PathHashIndexSize;
+		if (PathHashIndexOffset < 0 || 
+			PathHashIndexSize < 0 ||
+			PathHashIndexEndPosition.InvalidOrGreaterThan(CachedTotalSize) || 
+			IntFitsIn<int32>(PathHashIndexSize) == false)
+		{
+			// Should not be possible for these values (which came from the PrimaryIndex) to be invalid, since we verified the index hash of the PrimaryIndex
+			UE_LOG(LogPakFile, Log, TEXT("Corrupt pak PrimaryIndex detected!"));
+			UE_LOG(LogPakFile, Log, TEXT(" Filename: %s"), *PakFilename);
+			UE_LOG(LogPakFile, Log, TEXT(" Total Size: %lld"), Reader.TotalSize());
+			UE_LOG(LogPakFile, Log, TEXT(" PathHashIndexOffset : %lld"), PathHashIndexOffset);
+			UE_LOG(LogPakFile, Log, TEXT(" PathHashIndexSize: %lld"), PathHashIndexSize);
+			return false;
+		}
+		Reader.Seek(PathHashIndexOffset);
+		PathHashIndexData.SetNum((int32)(PathHashIndexSize));
+		{
+			SCOPED_BOOT_TIMING("PakFile_LoadPathHashIndex");
+			Reader.Serialize(PathHashIndexData.GetData(), PathHashIndexSize);
+		}
+
+		{
+			SCOPED_BOOT_TIMING("PakFile_HashPathHashIndex");
+			if (!DecryptAndValidateIndex(Reader, PathHashIndexData, PathHashIndexHash, ComputedHash))
+			{
+				UE_LOG(LogPakFile, Log, TEXT("Corrupt pak PathHashIndex detected!"));
+				UE_LOG(LogPakFile, Log, TEXT(" Filename: %s"), *PakFilename);
+				UE_LOG(LogPakFile, Log, TEXT(" Encrypted: %d"), Info.bEncryptedIndex);
+				UE_LOG(LogPakFile, Log, TEXT(" Total Size: %lld"), Reader.TotalSize());
+				UE_LOG(LogPakFile, Log, TEXT(" Index Offset: %lld"), FullDirectoryIndexOffset);
+				UE_LOG(LogPakFile, Log, TEXT(" Index Size: %lld"), FullDirectoryIndexSize);
+				UE_LOG(LogPakFile, Log, TEXT(" Stored Index Hash: %s"), *PathHashIndexHash.ToString());
+				UE_LOG(LogPakFile, Log, TEXT(" Computed Index Hash: %s"), *ComputedHash.ToString());
+				return false;
+			}
+		}
+
+		{
+			SCOPED_BOOT_TIMING("PakFile_SerializePathHashIndex");
+			PathHashIndexReader << PathHashIndex;
+		}
+		bHasPathHashIndex = true;
+	}
+	
+	if (!bReadFullDirectoryIndex)
+	{
+		check(bWillUsePathHashIndex); // Need to confirm that we have read the PathHashIndex bytes
+		// Store the PrunedDirectoryIndex in our DirectoryIndex
+		{
+			SCOPED_BOOT_TIMING("PakFile_SerializePrunedDirectoryIndex");
+			PathHashIndexReader << DirectoryIndex;
+		}
+		bHasFullDirectoryIndex = false;
+#if ENABLE_PAKFILE_RUNTIME_PRUNING
+		bWillPruneDirectoryIndex = false;
+#endif
+	}
+	else
+	{
+		FGuardedInt64 FullDirectoryIndexEndPosition = FGuardedInt64(FullDirectoryIndexOffset) + FullDirectoryIndexSize;
+		if (FullDirectoryIndexOffset  < 0 || 
+			FullDirectoryIndexSize < 0 ||
+			FullDirectoryIndexEndPosition.InvalidOrGreaterThan(CachedTotalSize) || 
+			IntFitsIn<int32>(FullDirectoryIndexSize) == false)
+		{
+			// Should not be possible for these values (which came from the PrimaryIndex) to be invalid, since we verified the index hash of the PrimaryIndex
+			UE_LOG(LogPakFile, Log, TEXT("Corrupt pak PrimaryIndex detected!"));
+			UE_LOG(LogPakFile, Log, TEXT(" Filename: %s"), *PakFilename);
+			UE_LOG(LogPakFile, Log, TEXT(" Total Size: %lld"), Reader.TotalSize());
+			UE_LOG(LogPakFile, Log, TEXT(" FullDirectoryIndexOffset : %lld"), FullDirectoryIndexOffset );
+			UE_LOG(LogPakFile, Log, TEXT(" FullDirectoryIndexSize: %lld"), FullDirectoryIndexSize);
+			return false;
+		}
+		TArray<uint8> FullDirectoryIndexData;
+		Reader.Seek(FullDirectoryIndexOffset );
+		FullDirectoryIndexData.SetNum((int32)(FullDirectoryIndexSize));
+		{
+			SCOPED_BOOT_TIMING("PakFile_LoadDirectoryIndex");
+			Reader.Serialize(FullDirectoryIndexData.GetData(), FullDirectoryIndexSize);
+		}
+
+		{
+			SCOPED_BOOT_TIMING("PakFile_HashDirectoryIndex");
+			if (!DecryptAndValidateIndex(Reader, FullDirectoryIndexData, FullDirectoryIndexHash, ComputedHash))
+			{
+				UE_LOG(LogPakFile, Log, TEXT("Corrupt pak FullDirectoryIndex detected!"));
+				UE_LOG(LogPakFile, Log, TEXT(" Filename: %s"), *PakFilename);
+				UE_LOG(LogPakFile, Log, TEXT(" Encrypted: %d"), Info.bEncryptedIndex);
+				UE_LOG(LogPakFile, Log, TEXT(" Total Size: %lld"), Reader.TotalSize());
+				UE_LOG(LogPakFile, Log, TEXT(" Index Offset: %lld"), FullDirectoryIndexOffset);
+				UE_LOG(LogPakFile, Log, TEXT(" Index Size: %lld"), FullDirectoryIndexSize);
+				UE_LOG(LogPakFile, Log, TEXT(" Stored Index Hash: %s"), *FullDirectoryIndexHash.ToString());
+				UE_LOG(LogPakFile, Log, TEXT(" Computed Index Hash: %s"), *ComputedHash.ToString());
+				return false;
+			}
+		}
+
+		FMemoryReader SecondaryIndexReader(FullDirectoryIndexData);
+		{
+			SCOPED_BOOT_TIMING("PakFile_SerializeDirectoryIndex");
+			SecondaryIndexReader << DirectoryIndex;
+		}
+		bHasFullDirectoryIndex = true;
+
+#if ENABLE_PAKFILE_RUNTIME_PRUNING
+		if (bWillUseFullDirectoryIndex)
+		{
+			bWillPruneDirectoryIndex = false;
+		}
+		else
+		{
+			// Store the PrunedDirectoryIndex from the PrimaryIndex in our PrunedDirectoryIndex, to be used for verification and to be swapped into DirectoryIndex later
+			check(bWillUsePathHashIndex); // Need to confirm that we have read the PathHashIndex bytes
+			{
+				SCOPED_BOOT_TIMING("PakFile_SerializePrunedDirectoryIndex");
+				PathHashIndexReader << PrunedDirectoryIndex;
+			}
+			bWillPruneDirectoryIndex = true;
+			bSomePakNeedsPruning = true;
+		}
+#endif
+	}
+
+	UE_LOG(LogPakFile, Verbose, TEXT("PakFile PrimaryIndexSize=%d"), Info.IndexSize);
+	UE_LOG(LogPakFile, Verbose, TEXT("PakFile PathHashIndexSize=%d"), PathHashIndexSize);
+	UE_LOG(LogPakFile, Verbose, TEXT("PakFile FullDirectoryIndexSize=%d"), FullDirectoryIndexSize);
+
+	check(bHasFullDirectoryIndex || bHasPathHashIndex);
+	return true;
+}
+
+bool FPakFile::LoadLegacyIndex(FArchive& Reader)
+{
+	bHasPathHashIndex = false;
+	bHasFullDirectoryIndex = false;
+#if ENABLE_PAKFILE_RUNTIME_PRUNING
+	bNeedsLegacyPruning = false;
+	bWillPruneDirectoryIndex = false;
+#endif
+
+	// Load index into memory first.
+	FGuardedInt64 IndexEndPosition = FGuardedInt64(Info.IndexOffset) + Info.IndexSize;
+	if (Info.IndexSize < 0 ||
+		Info.IndexOffset < 0 ||
+		IndexEndPosition.InvalidOrGreaterThan(CachedTotalSize) ||
+		IntFitsIn<int32>(Info.IndexSize) == false)
+	{
+		UE_LOG(LogPakFile, Fatal, TEXT("Corrupted index offset/size in pak file."));
+		return false;
+	}
+
+	TArray<uint8> IndexData;
+	IndexData.SetNum((int32)(Info.IndexSize));
+
+	Reader.Seek(Info.IndexOffset);
+	Reader.Serialize(IndexData.GetData(), Info.IndexSize);
+
+	FSHAHash ComputedHash;
+	if (!DecryptAndValidateIndex(Reader, IndexData, Info.IndexHash, ComputedHash))
+	{
+		UE_LOG(LogPakFile, Log, TEXT("Corrupt pak index detected!"));
+		UE_LOG(LogPakFile, Log, TEXT(" Filename: %s"), *PakFilename);
+		UE_LOG(LogPakFile, Log, TEXT(" Encrypted: %d"), Info.bEncryptedIndex);
+		UE_LOG(LogPakFile, Log, TEXT(" Total Size: %lld"), Reader.TotalSize());
+		UE_LOG(LogPakFile, Log, TEXT(" Index Offset: %lld"), Info.IndexOffset);
+		UE_LOG(LogPakFile, Log, TEXT(" Index Size: %lld"), Info.IndexSize);
+		UE_LOG(LogPakFile, Log, TEXT(" Stored Index Hash: %s"), *Info.IndexHash.ToString());
+		UE_LOG(LogPakFile, Log, TEXT(" Computed Index Hash: %s"), *ComputedHash.ToString());
+		return false;
+	}
+
+
+	FMemoryReader IndexReader(IndexData);
+
+	// Read the default mount point and all entries.
+	NumEntries = 0;
+	IndexReader << MountPoint;
+
+	// We are just deserializing a string, which could get however long. Since we know it's a path
+	// and paths are bound to file system rules, we know it can't get absurdly long (e.g. windows is _at best_ 32k)
+	// we just sanity check it to prevent operating on massive buffers and risking overflows.
+	if (MountPoint.Len() > 65535)
+	{
+		UE_LOG(LogPakFile, Error, TEXT("Corrupt pak index data: MountPoint path is longer than 65k"));
+		return false;
+	}
+	IndexReader << NumEntries;
+
+	if (NumEntries < 0)
+	{
+		UE_LOG(LogPakFile, Error, TEXT("Corrupt pak index data: NumEntries is negative"));
+		return false;
+	}
+
+	MakeDirectoryFromPath(MountPoint);
+
+	FPakEntryPair PakEntryPair;
+	auto ReadNextEntry = [&PakEntryPair, &IndexReader, this]() -> FPakEntryPair&
+	{
+		IndexReader << PakEntryPair.Filename;
+		PakEntryPair.Info.Reset();
+		PakEntryPair.Info.Serialize(IndexReader, this->Info.Version);
+		return PakEntryPair;
+	};
+
+	TMap<uint64, FString> CollisionDetection;
+	int32 NumEncodedEntries = 0;
+	int32 NumDeletedEntries = 0;
+#if ENABLE_PAKFILE_RUNTIME_PRUNING
+	bool bCreatePathHash = !IsPakKeepFullDirectory();
+#else
+	// Pruning of legacy files is no longer supported; we will keep the entire directory with no way to prune it.  There is no need to create the PathHashIndex since we will have the FullDirectoryIndex.
+	bool bCreatePathHash = false;
+#endif
+	FPathHashIndex* PathHashToWrite = bCreatePathHash ? &PathHashIndex : nullptr;
+	FPakFile::EncodePakEntriesIntoIndex(NumEntries, ReadNextEntry, *PakFilename, Info, MountPoint, NumEncodedEntries, NumDeletedEntries, &PathHashSeed,
+		&DirectoryIndex, PathHashToWrite, EncodedPakEntries, Files, &CollisionDetection, Info.Version);
+	check(NumEncodedEntries + Files.Num() + NumDeletedEntries == NumEntries);
+	Files.Shrink();
+	EncodedPakEntries.Shrink();
+
+	bHasPathHashIndex = bCreatePathHash;
+	bHasFullDirectoryIndex = true;
+#if ENABLE_PAKFILE_RUNTIME_PRUNING
+	if (!IsPakKeepFullDirectory())
+	{
+		bNeedsLegacyPruning = true;
+		bWillPruneDirectoryIndex = true;
+		bSomePakNeedsPruning = true;
+		// We cannot prune during this call because config settings have not yet been loaded and we need the settings for DirectoryIndexKeepFiles before we can prune
+		// PrunedDirectoryIndex will be created and swapped with DirectoryIndex in OptimizeMemoryUsageForMountedPaks, and bHasFullDirectoryIndex will be set to false then
+	}
+	else
+	{
+		bNeedsLegacyPruning = false;
+		bWillPruneDirectoryIndex = false;
+	}
+#endif
+
+	check(bHasFullDirectoryIndex || bHasPathHashIndex);
+	return true;
+}
+
+bool FPakFile::DecryptAndValidateIndex(FArchive& Reader, TArray<uint8>& IndexData, FSHAHash& InExpectedHash, FSHAHash& OutActualHash)
+{
+	// Decrypt if necessary
+	if (Info.bEncryptedIndex)
+	{
+		DecryptData(IndexData.GetData(), IndexData.Num(), Info.EncryptionKeyGuid);
+	}
+
+	// Check SHA1 value.
+	FSHA1::HashBuffer(IndexData.GetData(), IndexData.Num(), OutActualHash.Hash);
+	return InExpectedHash == OutActualHash;
+}
+
+uint64 FPakFile::HashPath(const TCHAR* RelativePathFromMount, uint64 Seed, int32 PakFileVersion)
+{
+	FString LowercaseRelativePath(RelativePathFromMount);
+	LowercaseRelativePath.ToLowerInline();
+	if (PakFileVersion >= FPakInfo::PakFile_Version_Fnv64BugFix)
+	{
+		return FFnv::MemFnv64(*LowercaseRelativePath, LowercaseRelativePath.Len() * sizeof(TCHAR), Seed);
+	}
+	else
+	{
+		return LegacyMemFnv64(*LowercaseRelativePath, LowercaseRelativePath.Len() * sizeof(TCHAR), Seed);
+	}
+}
+
+void FPakFile::EncodePakEntriesIntoIndex(int32 InNumEntries, const ReadNextEntryFunction& InReadNextEntry, const TCHAR* InPakFilename, const FPakInfo& InPakInfo, const FString& MountPoint,
+	int32& OutNumEncodedEntries, int32& OutNumDeletedEntries, uint64* OutPathHashSeed,
+	FDirectoryIndex* OutDirectoryIndex, FPathHashIndex* OutPathHashIndex, TArray<uint8>& OutEncodedPakEntries, TArray<FPakEntry>& OutNonEncodableEntries,
+	TMap<uint64, FString>* InOutCollisionDetection, int32 PakFileVersion)
+{
+	uint64 PathHashSeed = 0;
+	if (OutPathHashSeed || OutPathHashIndex)
+	{
+		FString LowercasePakFilename(InPakFilename);
+		LowercasePakFilename.ToLowerInline();
+		PathHashSeed = FCrc::StrCrc32(*LowercasePakFilename);
+		if (OutPathHashSeed)
+		{
+			*OutPathHashSeed = PathHashSeed;
+		}
+	}
+
+	OutNumEncodedEntries = 0;
+	OutNumDeletedEntries = 0;
+	FMemoryWriter CompressedEntryWriter(OutEncodedPakEntries);
+	for (int32 EntryCount = 0; EntryCount < InNumEntries; ++EntryCount)
+	{
+		FPakEntryPair& Pair = InReadNextEntry();
+		// Add the Entry and get an FPakEntryLocation for it
+		const FPakEntry& PakEntry = Pair.Info;
+		FPakEntryLocation EntryLocation;
+		if (!PakEntry.IsDeleteRecord())
+		{
+			EntryLocation = FPakEntryLocation::CreateFromOffsetIntoEncoded(OutEncodedPakEntries.Num());
+			if (EncodePakEntry(CompressedEntryWriter, PakEntry, InPakInfo))
+			{
+				++OutNumEncodedEntries;
+			}
+			else
+			{
+				int32 ListIndex = OutNonEncodableEntries.Num();
+				EntryLocation = FPakEntryLocation::CreateFromListIndex(ListIndex);
+				OutNonEncodableEntries.Add(PakEntry);
+
+				// PakEntries in the index have some values that are different from the in-place pakentries stored next to each file's payload.  EncodePakEntry handles that internally if encoding succeeded.
+				FPakEntry& StoredPakEntry = OutNonEncodableEntries[ListIndex];
+				FMemory::Memset(StoredPakEntry.Hash, 0); // Hash is 0-filled
+				StoredPakEntry.Verified = true; // Validation of the hash is impossible, so Verified is set to true
+			}
+		}
+		else
+		{
+			++OutNumDeletedEntries;
+		}
+
+		// Add the Entry into the requested Indexes
+		AddEntryToIndex(Pair.Filename, EntryLocation, MountPoint, PathHashSeed, OutDirectoryIndex, OutPathHashIndex, InOutCollisionDetection, PakFileVersion);
+	}
+}
+
+void FPakFile::PruneDirectoryIndex(FDirectoryIndex& InOutDirectoryIndex, FDirectoryIndex* PrunedDirectoryIndex, const FString& MountPoint)
+{
+	// Caller holds WriteLock on DirectoryIndexLock
+	TArray<FString> FileWildCards, DirectoryWildCards, OldWildCards;
+	GConfig->GetArray(TEXT("Pak"), TEXT("DirectoryIndexKeepFiles"), FileWildCards, GEngineIni);
+	GConfig->GetArray(TEXT("Pak"), TEXT("DirectoryIndexKeepEmptyDirectories"), DirectoryWildCards, GEngineIni);
+	GConfig->GetArray(TEXT("Pak"), TEXT("DirectoryRootsToKeepInMemoryWhenUnloadingPakEntryFilenames"), OldWildCards, GEngineIni); // Legacy name, treated as both KeepFiles and KeepEmptyDirectories
+	DirectoryWildCards.Append(OldWildCards);
+	FileWildCards.Append(OldWildCards);
+	int32 NumKeptEntries = 0;
+
+	if (PrunedDirectoryIndex)
+	{
+		PrunedDirectoryIndex->Empty();
+	}
+
+	TMap<FString, bool> KeepDirectory;
+
+	// Clear out those portions of the Index allowed by the user.
+	if (DirectoryWildCards.Num() > 0 || FileWildCards.Num() > 0)
+	{
+		for (auto It = InOutDirectoryIndex.CreateIterator(); It; ++It)
+		{
+			const FString& DirectoryName = It.Key();
+			FPakDirectory& OriginalDirectory = It.Value();
+			const FString FullDirectoryName = PakPathCombine(MountPoint, DirectoryName);
+			check(IsPathInDirectoryFormat(FullDirectoryName));
+			FPakDirectory* PrunedDirectory = nullptr;
+			bool bKeepDirectory = false;
+
+			TArray<FString> RemoveFilenames;
+			for (auto FileIt = It->Value.CreateIterator(); FileIt; ++FileIt)
+			{
+				const FString& FileNameWithoutPath = FileIt->Key;
+				const FString FullFilename = PakPathCombine(FullDirectoryName, FileNameWithoutPath);
+				bool bKeepFile = false;
+
+				for (const FString& WildCard : FileWildCards)
+				{
+					if (FullFilename.MatchesWildcard(WildCard))
+					{
+						bKeepFile = true;
+						break;
+					}
+				}
+
+				if (bKeepFile)
+				{
+					bKeepDirectory = true;
+					if (PrunedDirectoryIndex)
+					{
+						if (!PrunedDirectory)
+						{
+							PrunedDirectory = &PrunedDirectoryIndex->Add(DirectoryName);
+						}
+						PrunedDirectory->Add(FileNameWithoutPath, FileIt->Value);
+					}
+				}
+				else
+				{
+					if (!PrunedDirectoryIndex)
+					{
+						RemoveFilenames.Add(FileNameWithoutPath);
+					}
+				}
+			}
+			if (!PrunedDirectoryIndex)
+			{
+				for (const FString& FileNameWithoutPath : RemoveFilenames)
+				{
+					OriginalDirectory.Remove(FileNameWithoutPath);
+				}
+			}
+
+			if (!bKeepDirectory)
+			{
+				for (const FString& WildCard : DirectoryWildCards)
+				{
+					if (FullDirectoryName.MatchesWildcard(WildCard))
+					{
+						bKeepDirectory = true;
+						break;
+					}
+				}
+			}
+			KeepDirectory.FindOrAdd(DirectoryName) = bKeepDirectory;
+		}
+
+		// For each kept directory, mark that we need to keep all of its parents up to the mount point
+		for (const TPair<FString, bool>& Pair : KeepDirectory)
+		{
+			if (Pair.Value)
+			{
+				FString CurrentDirectory = Pair.Key;
+				FString UnusedCleanFileName;
+				while (CurrentDirectory != MountPoint)
+				{
+					if (!SplitPathInline(CurrentDirectory, UnusedCleanFileName))
+					{
+						break;
+					}
+					bool& bOldValue = KeepDirectory.FindOrAdd(CurrentDirectory);
+					if (bOldValue)
+					{
+						break;
+					}
+					bOldValue = true;
+				}
+			}
+		}
+
+		// Prune all of the directories
+		for (const TPair<FString, bool>& Pair : KeepDirectory)
+		{
+			const FString& DirectoryName = Pair.Key;
+			bool bKeep = Pair.Value;
+			if (bKeep)
+			{
+				if (PrunedDirectoryIndex)
+				{
+					PrunedDirectoryIndex->FindOrAdd(DirectoryName);
+				}
+			}
+			else
+			{
+				if (!PrunedDirectoryIndex)
+				{
+					InOutDirectoryIndex.Remove(DirectoryName);
+				}
+			}
+		}
+	}
+	else
+	{
+		if (!PrunedDirectoryIndex)
+		{
+			InOutDirectoryIndex.Empty();
+		}
+	}
+}
+
+FPakFile::EFindResult FPakFile::GetPakEntry(const FPakEntryLocation& PakEntryLocation, FPakEntry* OutEntry) const
+{
+	return GetPakEntry(PakEntryLocation, OutEntry, EncodedPakEntries, Files, Info);
+}
+
+FPakFile::EFindResult FPakFile::GetPakEntry(const FPakEntryLocation& PakEntryLocation, FPakEntry* OutEntry, const TArray<uint8>& EncodedPakEntries, const TArray<FPakEntry>& Files, const FPakInfo& Info)
+{
+	bool bDeleted = PakEntryLocation.IsInvalid();
+	if (OutEntry != NULL)
+	{
+		if (!bDeleted)
+		{
+			// The FPakEntry structures are bit-encoded, so decode it.
+			int32 EncodedOffset = PakEntryLocation.GetAsOffsetIntoEncoded();
+			if (EncodedOffset >= 0)
+			{
+				check(EncodedOffset < EncodedPakEntries.Num());
+				DecodePakEntry(EncodedPakEntries.GetData() + EncodedOffset, *OutEntry, Info);
+			}
+			else
+			{
+				int32 ListIndex = PakEntryLocation.GetAsListIndex();
+				check(ListIndex >= 0);
+				const FPakEntry& FoundEntry = Files[ListIndex];
+				//*OutEntry = **FoundEntry;
+				OutEntry->Offset = FoundEntry.Offset;
+				OutEntry->Size = FoundEntry.Size;
+				OutEntry->UncompressedSize = FoundEntry.UncompressedSize;
+				OutEntry->CompressionMethodIndex = FoundEntry.CompressionMethodIndex;
+				OutEntry->CompressionBlocks = FoundEntry.CompressionBlocks;
+				OutEntry->CompressionBlockSize = FoundEntry.CompressionBlockSize;
+				OutEntry->Flags = FoundEntry.Flags;
+			}
+		}
+		else
+		{
+			// entry was deleted, build dummy entry to indicate the deleted entry
+			(*OutEntry) = FPakEntry();
+			OutEntry->SetDeleteRecord(true);
+		}
+
+		// Index PakEntries do not store their hash, so verification against the hash is impossible.
+		// Initialize the OutEntry's Hash to 0 and mark it as already verified.
+		// TODO: Verified and Hash are checked by FPakFileHandle, which is used when opening files from PakFiles synchronously;
+		//       it is not currently used by asynchronous reads in FPakPrecacher, and we can likely remove it from use in FPakFileHandle as well
+		FMemory::Memset(OutEntry->Hash, 0);
+		OutEntry->Verified = true;
+	}
+
+	return bDeleted ? EFindResult::FoundDeleted : EFindResult::Found;
+}
+
+FPakFile::FIndexSettings& FPakFile::GetIndexSettings()
+{
+	static FIndexSettings IndexLoadParams;
+	return IndexLoadParams;
+}
+
+bool FPakFile::IsPakKeepFullDirectory()
+{
+	FIndexSettings& IndexLoadParams = GetIndexSettings();
+	return IndexLoadParams.bKeepFullDirectory;
+}
+
+bool FPakFile::IsPakValidatePruning()
+{
+#if ENABLE_PAKFILE_RUNTIME_PRUNING_VALIDATE
+	FIndexSettings& IndexLoadParams = GetIndexSettings();
+	return IndexLoadParams.bValidatePruning;
+#else
+	return false;
+#endif
+}
+
+bool FPakFile::IsPakDelayPruning()
+{
+	FIndexSettings& IndexLoadParams = GetIndexSettings();
+	return IndexLoadParams.bDelayPruning;
+}
+
+bool FPakFile::IsPakWritePathHashIndex()
+{
+	FIndexSettings& IndexLoadParams = GetIndexSettings();
+	return IndexLoadParams.bWritePathHashIndex;
+}
+
+bool FPakFile::IsPakWriteFullDirectoryIndex()
+{
+	FIndexSettings& IndexLoadParams = GetIndexSettings();
+	return IndexLoadParams.bWriteFullDirectoryIndex;
+}
+
+bool FPakFile::RequiresDirectoryIndexLock() const
+{
+#if ENABLE_PAKFILE_RUNTIME_PRUNING
+	return bWillPruneDirectoryIndex;
+#else
+	return false; 
+#endif
+}
+
+bool FPakFile::ShouldValidatePrunedDirectory() const
+{
+#if ENABLE_PAKFILE_RUNTIME_PRUNING_VALIDATE
+	return IsPakValidatePruning() && bWillPruneDirectoryIndex && !bNeedsLegacyPruning;
+#else
+	return false;
+#endif
+}
+
+void FPakFile::AddEntryToIndex(const FString& Filename, const FPakEntryLocation& EntryLocation, const FString& MountPoint, uint64 PathHashSeed,
+	FDirectoryIndex* DirectoryIndex, FPathHashIndex* PathHashIndex, TMap<uint64, FString>* CollisionDetection, int32 PakFileVersion)
+{
+	FString RelativePathFromMount;
+	if (FPaths::IsRelative(Filename))
+	{
+		RelativePathFromMount = Filename;
+	}
+	else
+	{
+		check(IsPathInDirectoryFormat(MountPoint));
+		RelativePathFromMount = Filename;
+		check(Filename.Len() > MountPoint.Len());
+		bool bSucceeded = GetRelativePathFromMountInline(RelativePathFromMount, MountPoint);
+		check(bSucceeded);
+	}
+
+	if (DirectoryIndex)
+	{
+		FString RelativeDirectoryFromMount = RelativePathFromMount;
+		FString CleanFileName;
+		SplitPathInline(RelativeDirectoryFromMount, CleanFileName);
+		FPakDirectory* Directory = DirectoryIndex->Find(RelativeDirectoryFromMount);
+		if (Directory == nullptr)
+		{
+			// add the parent directories up to the mount point (mount point relative path is "/")
+			FString CurrentDirectory(RelativeDirectoryFromMount);
+			FString UnusedCleanFileName;
+			while (!CurrentDirectory.IsEmpty())
+			{
+				// FPaths::GetPath doesn't handle our / at the end of directories, so call FindLastChar(/)
+				if (!SplitPathInline(CurrentDirectory, UnusedCleanFileName))
+				{
+					break;
+				}
+				DirectoryIndex->FindOrAdd(CurrentDirectory);
+			}
+
+			// Add the new directory; this has to come after the addition of the parent directories because we want to use the pointer afterwards and adding other directories would invalidate it
+			Directory = &DirectoryIndex->Add(RelativeDirectoryFromMount);
+		}
+		Directory->Add(CleanFileName, EntryLocation);
+	}
+
+	// Add the entry into the PathHash index
+	if (CollisionDetection || PathHashIndex)
+	{
+		uint64 PathHash = FPakFile::HashPath(*RelativePathFromMount, PathHashSeed, PakFileVersion);
+		if (CollisionDetection)
+		{
+			FString* ExistingFilename = CollisionDetection->Find(PathHash);
+			checkf(!ExistingFilename || ExistingFilename->Equals(RelativePathFromMount), TEXT("Hash collision for two Filenames within a PakFile.  Filename1 == '%s'.  Filename2 == '%s'.  Hash='0x%x'.")
+				TEXT(" Collision handling is not yet implemented; to resolve the conflict, modify one of the Filenames."), **ExistingFilename, *RelativePathFromMount, PathHash);
+			CollisionDetection->Add(PathHash, RelativePathFromMount);
+		}
+		if (PathHashIndex)
+		{
+			PathHashIndex->Add(PathHash, EntryLocation);
+		}
+	}
+}
+
+// Take a pak entry and byte encode it into a smaller representation
+bool FPakFile::EncodePakEntry(FArchive& Ar, const FPakEntry& InPakEntry, const FPakInfo& InInfo)
+{
+	// See notes in DecodePakEntry for the output layout
+
+	check(Ar.IsSaving()); // This function is encode only, and promises not to modify InEntry, but we want to use << which takes non-const so we need to assert that Ar is a loader not a saver
+	check(!InPakEntry.IsDeleteRecord()); // Deleted PakEntries cannot be encoded, caller must check for IsDeleteRecord and handle it separately by e.g. not adding the Entry to the FileList and instead giving the referencer an invalid FPakEntryLocation
+	FPakEntry& PakEntry = const_cast<FPakEntry&>(InPakEntry);
+
+	const uint32 CompressedBlockAlignment = PakEntry.IsEncrypted() ? FAES::AESBlockSize : 1;
+	const int64 HeaderSize = PakEntry.GetSerializedSize(InInfo.Version);
+
+	// This data fits into a bitfield (described in DecodePakEntry), and the data has
+	// to fit within a certain range of bits.
+	if (PakEntry.CompressionMethodIndex >= (1 << 6))
+	{
+		return false;
+	}
+	if (PakEntry.CompressionBlocks.Num() >= (1 << 16))
+	{
+		return false;
+	}
+	if (PakEntry.CompressionMethodIndex != 0)
+	{
+		if (PakEntry.CompressionBlocks.Num() > 0 && ((InInfo.HasRelativeCompressedChunkOffsets() ? 0 : PakEntry.Offset) + HeaderSize != PakEntry.CompressionBlocks[0].CompressedStart))
+		{
+			return false;
+		}
+		if (PakEntry.CompressionBlocks.Num() == 1)
+		{
+			uint64 Base = InInfo.HasRelativeCompressedChunkOffsets() ? 0 : PakEntry.Offset;
+			uint64 AlignedBlockSize = Align(PakEntry.CompressionBlocks[0].CompressedEnd - PakEntry.CompressionBlocks[0].CompressedStart, CompressedBlockAlignment);
+			if ((Base + HeaderSize + PakEntry.Size) != (PakEntry.CompressionBlocks[0].CompressedStart + AlignedBlockSize))
+			{
+				return false;
+			}
+		}
+		if (PakEntry.CompressionBlocks.Num() > 1)
+		{
+			for (int i = 1; i < PakEntry.CompressionBlocks.Num(); ++i)
+			{
+				uint64 PrevBlockSize = PakEntry.CompressionBlocks[i - 1].CompressedEnd - PakEntry.CompressionBlocks[i - 1].CompressedStart;
+				PrevBlockSize = Align(PrevBlockSize, CompressedBlockAlignment);
+				if (PakEntry.CompressionBlocks[i].CompressedStart != (PakEntry.CompressionBlocks[i - 1].CompressedStart + PrevBlockSize))
+				{
+					return false;
+				}
+			}
+		}
+	}
+
+	// This entry can be encoded, so let's do it!
+
+	bool bIsOffset32BitSafe = PakEntry.Offset <= MAX_uint32;
+	bool bIsSize32BitSafe = PakEntry.Size <= MAX_uint32;
+	bool bIsUncompressedSize32BitSafe = PakEntry.UncompressedSize <= MAX_uint32;
+	
+	// If CompressionBlocks.Num() == 1, we set CompressionBlockSize == UncompressedSize and record CompressBlockSizePacked=0
+	// Otherwise, we encode CompressionBlockSize as a 6-bit multiple of 1 << 11.
+	// If CompressionBlockSize is not a multiple of 1 << 11, or is a larger multiple than 6 bits we can not encode correctly.
+	// In that case we set the packed field to its maximum value (0x3F) and send the unencoded CompressionBlockSize as a separate value.
+	uint32 CompressionBlockSizePacked = 0;
+	if (PakEntry.CompressionBlocks.Num() > 1)
+	{
+		CompressionBlockSizePacked = (PakEntry.CompressionBlockSize >> 11) & 0x3F;
+		if ((CompressionBlockSizePacked << 11) != PakEntry.CompressionBlockSize)
+		{
+			CompressionBlockSizePacked = 0x3F;
+		}
+	}
+
+	// Build the Flags field.
+	uint32 Flags =
+		(bIsOffset32BitSafe ? (1 << 31) : 0)
+		| (bIsUncompressedSize32BitSafe ? (1 << 30) : 0)
+		| (bIsSize32BitSafe ? (1 << 29) : 0)
+		| (PakEntry.CompressionMethodIndex << 23)
+		| (PakEntry.IsEncrypted() ? (1 << 22) : 0)
+		| (PakEntry.CompressionBlocks.Num() << 6)
+		| CompressionBlockSizePacked
+		;
+
+	Ar << Flags;
+	
+	// if we write 0x3F for CompressionBlockSize then send the field
+	if ( CompressionBlockSizePacked == 0x3F )
+	{
+		uint32 Value = (uint32)PakEntry.CompressionBlockSize;
+		Ar << Value;
+	}
+
+	// Build the Offset field.
+	if (bIsOffset32BitSafe)
+	{
+		uint32 Value = (uint32)PakEntry.Offset;
+		Ar << Value;
+	}
+	else
+	{
+		Ar << PakEntry.Offset;
+	}
+
+	// Build the Uncompressed Size field.
+	if (bIsUncompressedSize32BitSafe)
+	{
+		uint32 Value = (uint32)PakEntry.UncompressedSize;
+		Ar << Value;
+	}
+	else
+	{
+		Ar << PakEntry.UncompressedSize;
+	}
+
+	// Any additional data is for compressed file data.
+	if (PakEntry.CompressionMethodIndex != 0)
+	{
+		// Build the Compressed Size field.
+		if (bIsSize32BitSafe)
+		{
+			uint32 Value = (uint32)PakEntry.Size;
+			Ar << Value;
+		}
+		else
+		{
+			Ar << PakEntry.Size;
+		}
+
+		// Build the Compression Blocks array.
+		if (PakEntry.CompressionBlocks.Num() > 1 || (PakEntry.CompressionBlocks.Num() == 1 && PakEntry.IsEncrypted()))
+		{
+			for (int CompressionBlockIndex = 0; CompressionBlockIndex < PakEntry.CompressionBlocks.Num(); ++CompressionBlockIndex)
+			{
+				uint32 Value = IntCastChecked<uint32>(PakEntry.CompressionBlocks[CompressionBlockIndex].CompressedEnd - PakEntry.CompressionBlocks[CompressionBlockIndex].CompressedStart);
+				Ar << Value;
+			}
+		}
+	}
+
+	return true;
+}
+
+void FPakFile::DecodePakEntry(const uint8* SourcePtr, FPakEntry& OutEntry, const FPakInfo& InInfo)
+{
+	// Grab the big bitfield value:
+	// Bit 31 = Offset 32-bit safe?
+	// Bit 30 = Uncompressed size 32-bit safe?
+	// Bit 29 = Size 32-bit safe?
+	// Bits 28-23 = Compression method
+	// Bit 22 = Encrypted
+	// Bits 21-6 = Compression blocks count
+	// Bits 5-0 = Compression block size
+	uint32 Value = *(uint32*)SourcePtr;
+	SourcePtr += sizeof(uint32);
+	
+	uint32 CompressionBlockSize = 0;
+	if ( (Value & 0x3f) == 0x3f ) // flag value to load a field
+	{
+		CompressionBlockSize = *(uint32*)SourcePtr;
+		SourcePtr += sizeof(uint32);
+	}
+	else
+	{
+		// for backwards compatibility with old paks :
+		CompressionBlockSize = ((Value & 0x3f) << 11);
+	}
+
+	// Filter out the CompressionMethod.
+	OutEntry.CompressionMethodIndex = (Value >> 23) & 0x3f;
+
+	// Test for 32-bit safe values. Grab it, or memcpy the 64-bit value
+	// to avoid alignment exceptions on platforms requiring 64-bit alignment
+	// for 64-bit variables.
+	//
+	// Read the Offset.
+	bool bIsOffset32BitSafe = (Value & (1 << 31)) != 0;
+	if (bIsOffset32BitSafe)
+	{
+		OutEntry.Offset = *(uint32*)SourcePtr;
+		SourcePtr += sizeof(uint32);
+	}
+	else
+	{
+		FMemory::Memcpy(&OutEntry.Offset, SourcePtr, sizeof(int64));
+		SourcePtr += sizeof(int64);
+	}
+
+	// Read the UncompressedSize.
+	bool bIsUncompressedSize32BitSafe = (Value & (1 << 30)) != 0;
+	if (bIsUncompressedSize32BitSafe)
+	{
+		OutEntry.UncompressedSize = *(uint32*)SourcePtr;
+		SourcePtr += sizeof(uint32);
+	}
+	else
+	{
+		FMemory::Memcpy(&OutEntry.UncompressedSize, SourcePtr, sizeof(int64));
+		SourcePtr += sizeof(int64);
+	}
+
+	// Fill in the Size.
+	if (OutEntry.CompressionMethodIndex != 0)
+	{
+		// Size is only present if compression is applied.
+		bool bIsSize32BitSafe = (Value & (1 << 29)) != 0;
+		if (bIsSize32BitSafe)
+		{
+			OutEntry.Size = *(uint32*)SourcePtr;
+			SourcePtr += sizeof(uint32);
+		}
+		else
+		{
+			FMemory::Memcpy(&OutEntry.Size, SourcePtr, sizeof(int64));
+			SourcePtr += sizeof(int64);
+		}
+	}
+	else
+	{
+		// The Size is the same thing as the UncompressedSize when
+		// CompressionMethod == COMPRESS_None.
+		OutEntry.Size = OutEntry.UncompressedSize;
+	}
+
+	// Filter the encrypted flag.
+	OutEntry.SetEncrypted((Value & (1 << 22)) != 0);
+
+	// This should clear out any excess CompressionBlocks that may be valid in the user's
+	// passed in entry.
+	uint32 CompressionBlocksCount = (Value >> 6) & 0xffff;
+	OutEntry.CompressionBlocks.Empty(CompressionBlocksCount);
+	OutEntry.CompressionBlocks.SetNum(CompressionBlocksCount);
+	
+	OutEntry.CompressionBlockSize = 0;
+	if (CompressionBlocksCount > 0)
+	{
+		OutEntry.CompressionBlockSize = CompressionBlockSize;
+		// Per the comment in Encode, if CompressionBlocksCount == 1, we use UncompressedSize for CompressionBlockSize
+		if (CompressionBlocksCount == 1)
+		{
+			OutEntry.CompressionBlockSize = IntCastChecked<uint32>(OutEntry.UncompressedSize);
+		}
+		ensure(OutEntry.CompressionBlockSize != 0);
+	}
+
+	// Set bDeleteRecord to false, because it obviously isn't deleted if we are here.
+	OutEntry.SetDeleteRecord(false);
+
+	// Base offset to the compressed data
+	int64 BaseOffset = InInfo.HasRelativeCompressedChunkOffsets() ? 0 : OutEntry.Offset;
+
+	// Handle building of the CompressionBlocks array.
+	if (OutEntry.CompressionBlocks.Num() == 1 && !OutEntry.IsEncrypted())
+	{
+		// If the number of CompressionBlocks is 1, we didn't store any extra information.
+		// Derive what we can from the entry's file offset and size.
+		FPakCompressedBlock& CompressedBlock = OutEntry.CompressionBlocks[0];
+		CompressedBlock.CompressedStart = BaseOffset + OutEntry.GetSerializedSize(InInfo.Version);
+		CompressedBlock.CompressedEnd = CompressedBlock.CompressedStart + OutEntry.Size;
+	}
+	else if (OutEntry.CompressionBlocks.Num() > 0)
+	{
+		// Get the right pointer to start copying the CompressionBlocks information from.
+		uint32* CompressionBlockSizePtr = (uint32*)SourcePtr;
+
+		// Alignment of the compressed blocks
+		uint64 CompressedBlockAlignment = OutEntry.IsEncrypted() ? FAES::AESBlockSize : 1;
+
+		// CompressedBlockOffset is the starting offset. Everything else can be derived from there.
+		int64 CompressedBlockOffset = BaseOffset + OutEntry.GetSerializedSize(InInfo.Version);
+		for (int CompressionBlockIndex = 0; CompressionBlockIndex < OutEntry.CompressionBlocks.Num(); ++CompressionBlockIndex)
+		{
+			FPakCompressedBlock& CompressedBlock = OutEntry.CompressionBlocks[CompressionBlockIndex];
+			CompressedBlock.CompressedStart = CompressedBlockOffset;
+			CompressedBlock.CompressedEnd = CompressedBlockOffset + *CompressionBlockSizePtr++;
+			CompressedBlockOffset += Align(CompressedBlock.CompressedEnd - CompressedBlock.CompressedStart, CompressedBlockAlignment);
+		}
+	}
+}
+
+bool FPakFile::NormalizeDirectoryQuery(const TCHAR* InPath, FString& OutRelativePathFromMount) const
+{
+	OutRelativePathFromMount = InPath;
+	MakeDirectoryFromPath(OutRelativePathFromMount);
+	return GetRelativePathFromMountInline(OutRelativePathFromMount, MountPoint);
+}
+
+const FPakDirectory* FPakFile::FindPrunedDirectoryInternal(const FString& RelativePathFromMount) const
+{
+	// Caller holds FScopedPakDirectoryIndexAccess
+	const FPakDirectory* PakDirectory = nullptr;
+
+#if ENABLE_PAKFILE_RUNTIME_PRUNING_VALIDATE
+	if (ShouldValidatePrunedDirectory())
+	{
+		PakDirectory = DirectoryIndex.Find(RelativePathFromMount);
+		const FPakDirectory* PrunedPakDirectory = PrunedDirectoryIndex.Find(RelativePathFromMount);
+		if ((PakDirectory != nullptr) != (PrunedPakDirectory != nullptr))
+		{
+			TSet<FString> FullFoundFiles, PrunedFoundFiles;
+			FString ReportedDirectoryName = MountPoint + RelativePathFromMount;
+			if (PakDirectory)
+			{
+				FullFoundFiles.Add(ReportedDirectoryName);
+			}
+			if (PrunedPakDirectory)
+			{
+				PrunedFoundFiles.Add(ReportedDirectoryName);
+			}
+			ValidateDirectorySearch(FullFoundFiles, PrunedFoundFiles, *ReportedDirectoryName);
+		}
+	}
+	else
+#endif
+	{
+		PakDirectory = DirectoryIndex.Find(RelativePathFromMount);
+	}
+	return PakDirectory;
+}
+
+bool FPakFile::Check()
+{
+	UE_LOG(LogPakFile, Display, TEXT("Checking pak file \"%s\". This may take a while..."), *PakFilename);
+	double StartTime = FPlatformTime::Seconds();
+
+	FSharedPakReader PakReader = GetSharedReader(nullptr);
+	int32 ErrorCount = 0;
+	int32 FileCount = 0;
+
+	// If the pak file is signed, we can do a fast check by just reading a single byte from the start of
+	// each signing block. The signed archive reader will bring in that whole signing block and compare
+	// against the signature table and fire the handler
+	if (bSigned)
+	{
+		FDelegateHandle DelegateHandle;
+		FPakPlatformFile::FPakSigningFailureHandlerData& HandlerData = FPakPlatformFile::GetPakSigningFailureHandlerData();
+
+		{
+			FScopeLock Lock(&HandlerData.GetLock());
+			DelegateHandle = HandlerData.GetPakChunkSignatureCheckFailedDelegate().AddLambda([&ErrorCount](const FPakChunkSignatureCheckFailedData&)
+			{
+				++ErrorCount;
+			});
+		}
+
+		int64 CurrentPos = 0;
+		const int64 Size = PakReader->TotalSize();
+		while (CurrentPos < Size)
+		{
+			PakReader->Seek(CurrentPos);
+			uint8 Byte = 0;
+			PakReader.GetArchive() << Byte;
+			CurrentPos += FPakInfo::MaxChunkDataSize;
+		}
+
+		if (DelegateHandle.IsValid())
+		{
+			FScopeLock Lock(&HandlerData.GetLock());
+			HandlerData.GetPakChunkSignatureCheckFailedDelegate().Remove(DelegateHandle);
+		}
+	}
+	else
+	{
+		const bool bIncludeDeleted = true;
+		TCHAR EntryNameBuffer[256];
+		auto GetEntryName = [&EntryNameBuffer](const FPakFile::FPakEntryIterator& It)
+		{
+			const FString* EntryFilename = It.TryGetFilename();
+			if (EntryFilename)
+			{
+				TCString<TCHAR>::Snprintf(EntryNameBuffer, sizeof(EntryNameBuffer), TEXT("\"%s\""), **EntryFilename);
+			}
+			else
+			{
+				TCString<TCHAR>::Snprintf(EntryNameBuffer, sizeof(EntryNameBuffer), TEXT("file at offset %u"), It.Info().Offset);
+			}
+			return EntryNameBuffer;
+		};
+		for (FPakFile::FPakEntryIterator It(*this, bIncludeDeleted); It; ++It, ++FileCount)
+		{
+			const FPakEntry& EntryFromIndex = It.Info();
+			if (EntryFromIndex.IsDeleteRecord())
+			{
+				UE_LOG(LogPakFile, Verbose, TEXT("%s Deleted."), GetEntryName(It));
+				continue;
+			}
+
+			void* FileContents = FMemory::Malloc(EntryFromIndex.Size);
+			PakReader->Seek(EntryFromIndex.Offset);
+			uint32 SerializedCrcTest = 0;
+			FPakEntry EntryFromPayload;
+			EntryFromPayload.Serialize(PakReader.GetArchive(), GetInfo().Version);
+			if (!EntryFromPayload.IndexDataEquals(EntryFromIndex))
+			{
+				UE_LOG(LogPakFile, Error, TEXT("Index FPakEntry does not match Payload FPakEntry for %s."), GetEntryName(It));
+				ErrorCount++;
+			}
+			PakReader->Serialize(FileContents, EntryFromIndex.Size);
+
+			uint8 TestHash[20];
+			FSHA1::HashBuffer(FileContents, EntryFromIndex.Size, TestHash);
+			if (FMemory::Memcmp(TestHash, EntryFromPayload.Hash, sizeof(TestHash)) != 0)
+			{
+				UE_LOG(LogPakFile, Error, TEXT("Hash mismatch for %s."), GetEntryName(It));
+				ErrorCount++;
+			}
+			else
+			{
+				UE_LOG(LogPakFile, Verbose, TEXT("%s OK. [%s]"), GetEntryName(It), *Info.GetCompressionMethod(EntryFromIndex.CompressionMethodIndex).ToString());
+			}
+			FMemory::Free(FileContents);
+		}
+		if (ErrorCount == 0)
+		{
+			UE_LOG(LogPakFile, Display, TEXT("Pak file \"%s\" healthy, %d files checked."), *PakFilename, FileCount);
+		}
+		else
+		{
+			UE_LOG(LogPakFile, Display, TEXT("Pak file \"%s\" corrupted (%d errors out of %d files checked.)."), *PakFilename, ErrorCount, FileCount);
+		}
+	}
+
+	double EndTime = FPlatformTime::Seconds();
+	double ElapsedTime = EndTime - StartTime;
+	UE_LOG(LogPakFile, Display, TEXT("Pak file \"%s\" checked in %.2fs"), *PakFilename, ElapsedTime);
+
+	return ErrorCount == 0;
+}
+
+void FPakFile::GetPrunedFilenames(TArray<FString>& OutFileList) const
+{
+	for (FFilenameIterator It(*this, true /* bIncludeDeleted */); It; ++It)
+	{
+		OutFileList.Add(PakPathCombine(MountPoint, It.Filename()));
+	}
+}
+
+void FPakFile::GetPrunedFilenamesInChunk(const TArray<int32>& InChunkIDs, TArray<FString>& OutFileList) const
+{
+	for (FFilenameIterator It(*this, true /* bIncludeDeleted */); It; ++It)
+	{
+		const FPakEntry& File = It.Info();
+		int64 FileStart = File.Offset;
+		int64 FileEnd = File.Offset + File.Size;
+
+		for (int64 LocalChunkID : InChunkIDs)
+		{
+			int64 ChunkStart = LocalChunkID * FPakInfo::MaxChunkDataSize;
+			int64 ChunkEnd = ChunkStart + FPakInfo::MaxChunkDataSize;
+
+			if (FileStart < ChunkEnd && FileEnd > ChunkStart)
+			{
+				OutFileList.Add(It.Filename());
+				break;
+			}
+		}
+	}
+}
+
+void FPakFile::FindPrunedFilesAtPath(const TCHAR* InPath, TArray<FString>& OutFiles,
+	bool bIncludeFiles, bool bIncludeDirectories, bool bRecursive) const
+{
+	auto ShouldVisit = [](FStringView Path) { return true; };
+	FindPrunedFilesAtPathInternal(InPath, ShouldVisit, OutFiles, bIncludeFiles, bIncludeDirectories, bRecursive);
+}
+
+
+#if ENABLE_PAKFILE_RUNTIME_PRUNING_VALIDATE
+void FPakFile::ValidateDirectorySearch(const TSet<FString>& FullFoundFiles, const TSet<FString>& PrunedFoundFiles, const TCHAR* InPath) const
+{
+	TArray<FString> MissingFromPruned;
+	for (const FString& FileInFull : FullFoundFiles)
+	{
+		if (!PrunedFoundFiles.Contains(FileInFull))
+		{
+			MissingFromPruned.Add(FileInFull);
+		}
+	}
+	TArray<FString> MissingFromFull;
+	for (const FString& FileInPruned : PrunedFoundFiles)
+	{
+		if (!FullFoundFiles.Contains(FileInPruned))
+		{
+			MissingFromFull.Add(FileInPruned);
+		}
+	}
+
+	if (MissingFromPruned.Num() == 0 && MissingFromFull.Num() == 0)
+	{
+		return;
+	}
+
+	TArray<FString> WildCards, OldWildCards;
+	GConfig->GetArray(TEXT("Pak"), TEXT("IndexValidationIgnore"), WildCards, GEngineIni);
+	auto IsIgnore = [&WildCards](const FString& FilePath)
+	{
+		for (const FString& WildCard : WildCards)
+		{
+			if (FilePath.MatchesWildcard(WildCard))
+			{
+				return true;
+			}
+		}
+		return false;
+	};
+	auto StripIgnores = [&IsIgnore](TArray<FString>& FilePaths)
+	{
+		for (int Idx = FilePaths.Num() - 1; Idx >= 0; --Idx)
+		{
+			if (IsIgnore(FilePaths[Idx]))
+			{
+				FilePaths.RemoveAtSwap(Idx);
+			}
+		}
+	};
+
+	StripIgnores(MissingFromPruned);
+	StripIgnores(MissingFromFull);
+
+	if (MissingFromPruned.Num() == 0 && MissingFromFull.Num() == 0)
+	{
+		return;
+	}
+	MissingFromPruned.Sort();
+	MissingFromFull.Sort();
+
+	// TODO: Restore this as an Error once we modify IPlatformFile::IterateDirectoryRecursively to declare its filefilter so we can ignore the spurious
+	// discovered files that are not part of the fully filtered query
+	UE_LOG(LogPakFile, Error, TEXT("FindPrunedFilesAtPath('%s') for PakFile '%s' found a different list in the FullDirectory than in the PrunedDirectory. ")
+		TEXT("Change the calling code or add the files to Engine:[Pak]:WildcardsToKeepInPakStringIndex or Engine:[Pak]:IndexValidationIgnore."),
+		InPath, *PakFilename);
+
+#if !NO_LOGGING && !UE_BUILD_SHIPPING
+	// Logging callstacks is expensive (multiple seconds long). Only do it the first time a path is seen, and only for the first
+	// few paths.
+	static TSet<FString> AlreadyLoggedCallstack;
+	static FCriticalSection AlreadyLoggedCallstackLock;
+	constexpr int32 CallstackLogDirsMax = 10;
+	bool bShouldLogCallstack = false;
+	if (AlreadyLoggedCallstack.Num() < CallstackLogDirsMax) // check to avoid taking critical section if unnecessary
+	{
+		FScopeLock AlreadyLoggedCallstackScopeLock(&AlreadyLoggedCallstackLock);
+		if (AlreadyLoggedCallstack.Num() < CallstackLogDirsMax) // check again since other thread may have modified it
+		{
+			bool bAlreadyLogged;
+			AlreadyLoggedCallstack.Add(FString(InPath), &bAlreadyLogged);
+			bShouldLogCallstack = !bAlreadyLogged;
+		}
+	}
+	if (bShouldLogCallstack)
+	{
+		UE_LOG(LogPakFile, Warning, TEXT("Callstack of FindPrunedFilesAtPath('%s'):"), InPath);
+		FDebug::DumpStackTraceToLog(ELogVerbosity::Warning);
+	}
+#endif
+
+	if (MissingFromPruned.Num() > 0)
+	{
+		for (const FString& Missing : MissingFromPruned)
+		{
+			UE_LOG(LogPakFile, Warning, TEXT("MissingPrunedPakFile: %s"), *Missing);
+		}
+	}
+	if (MissingFromFull.Num() > 0)
+	{
+		UE_LOG(LogPakFile, Error, TEXT("Some files in the PrunedDirectory are missing from the FullDirectory.  This is a logic error in FPakFile since the PrunedDirectory should be a subset of the FullDirectory."));
+		for (const FString& Missing : MissingFromFull)
+		{
+			UE_LOG(LogPakFile, Warning, TEXT("MissingFullPakFile: %s"), *Missing);
+		}
+	}
+}
+#endif
+
+bool FPakFile::RecreatePakReaders(IPlatformFile* LowerLevel)
+{
+	FScopeLock ScopedLock(&ReadersCriticalSection);
+
+	if (CurrentlyUsedReaders > 0)
+	{
+		UE_LOG(LogPakFile, Error, TEXT("Recreating pak readers while we have readers loaned out, this may be lead to crashes or decryption problems."));
+	}
+
+	// need to reset the decryptor as it will hold a pointer to the first created pak reader
+	Decryptor.Reset();
+
+	TArray<FArchiveAndLastAccessTime> TempReaders;
+
+	// Create a new PakReader *per* instance that was already mapped
+	for (const FArchiveAndLastAccessTime& Reader : Readers)
+	{
+		TUniquePtr<FArchive> PakReader = TUniquePtr<FArchive>(CreatePakReader(LowerLevel, *GetFilename()));
+		if (!PakReader)
+		{
+			UE_LOG(LogPakFile, Warning, TEXT("Unable to re-create pak \"%s\" handle"), *GetFilename());
+			return false;
+		}
+		TempReaders.Add(FArchiveAndLastAccessTime{ MoveTemp(PakReader), Reader.LastAccessTime });
+	}
+
+	// replace the current Readers with the newly created pak readers leaving them to out of scope
+	Readers= MoveTemp(TempReaders);
+
+	return true;
+}
+
+FSharedPakReader FPakFile::GetSharedReader(IPlatformFile* LowerLevel)
+{
+	LLM_SCOPE_BYTAG(PakSharedReaders);
+	FArchive* PakReader = nullptr;
+	{
+		FScopeLock ScopedLock(&ReadersCriticalSection);
+		if (Readers.Num())
+		{
+			FArchiveAndLastAccessTime Reader = Readers.Pop();
+			PakReader = Reader.Archive.Release();
+		}
+		else
+		{
+			// Create a new FArchive reader and pass it to the new handle.
+			PakReader = CreatePakReader(LowerLevel, *GetFilename());
+
+			if (!PakReader)
+			{
+				UE_LOG(LogPakFile, Warning, TEXT("Unable to create pak \"%s\" handle"), *GetFilename());
+			}
+		}
+		++CurrentlyUsedReaders;
+	}
+
+	return FSharedPakReader(PakReader, this);
+}
+
+void FPakFile::ReturnSharedReader(FArchive* Archive)
+{
+	FScopeLock ScopedLock(&ReadersCriticalSection);
+	--CurrentlyUsedReaders;
+	Readers.Push(FArchiveAndLastAccessTime{ TUniquePtr<FArchive>{Archive }, FPlatformTime::Seconds()});
+}
+
+void FPakFile::ReleaseOldReaders(double MaxAgeSeconds)
+{
+	if (ReadersCriticalSection.TryLock())
+	{
+		ON_SCOPE_EXIT
+		{
+			ReadersCriticalSection.Unlock();
+		};
+		double SearchTime = FPlatformTime::Seconds() - MaxAgeSeconds;
+		for (int32 i = Readers.Num() - 1; i >= 0; --i)
+		{
+			const FArchiveAndLastAccessTime& Reader = Readers[i];
+			if (Reader.LastAccessTime <= SearchTime)
+			{
+				// Remove this and all readers older than it (pushed before it)
+				Readers.RemoveAt(0, i + 1);
+				break;
+			}
+		}
+
+		if (Readers.Num() == 0 && CurrentlyUsedReaders == 0)
+		{
+			Decryptor.Reset();
+		}
+	}
+}
+
+const FPakEntryLocation* FPakFile::FindLocationFromIndex(const FString& FullPath, const FString& MountPoint, const FPathHashIndex& PathHashIndex, uint64 PathHashSeed, int32 PakFileVersion)
+{
+	const TCHAR* RelativePathFromMount = GetRelativeFilePathFromMountPointer(FullPath, MountPoint);
+	if (!RelativePathFromMount)
+	{
+		return nullptr;
+	}
+	uint64 PathHash = HashPath(RelativePathFromMount, PathHashSeed, PakFileVersion);
+	return PathHashIndex.Find(PathHash);
+}
+
+const FPakEntryLocation* FPakFile::FindLocationFromIndex(const FString& FullPath, const FString& MountPoint, const FDirectoryIndex& DirectoryIndex)
+{
+	if (!FullPath.StartsWith(MountPoint))
+	{
+		return nullptr;
+	}
+	const TCHAR* RelativePathFromMount = (*FullPath) + MountPoint.Len();
+	FString RelativeDirName(RelativePathFromMount);
+	FString CleanFileName;
+	if (RelativeDirName.IsEmpty())
+	{
+		return nullptr;
+	}
+	SplitPathInline(RelativeDirName, CleanFileName);
+	const FPakDirectory* PakDirectory = DirectoryIndex.Find(RelativeDirName);
+	if (PakDirectory)
+	{
+		return PakDirectory->Find(CleanFileName);
+	}
+	return nullptr;
+}
+
+FPakFile::EFindResult FPakFile::Find(const FString& FullPath, FPakEntry* OutEntry) const
+{
+	//QUICK_SCOPE_CYCLE_COUNTER(PakFileFind);
+
+	const FPakEntryLocation* PakEntryLocation;
+#if ENABLE_PAKFILE_RUNTIME_PRUNING_VALIDATE
+	if (IsPakValidatePruning() && bHasPathHashIndex && bHasFullDirectoryIndex)
+	{
+		const FPakEntryLocation* PathHashLocation = nullptr;
+		PathHashLocation = FindLocationFromIndex(FullPath, MountPoint, PathHashIndex, PathHashSeed, Info.Version);
+
+		const FPakEntryLocation* DirectoryLocation = nullptr;
+
+		{
+			FScopedPakDirectoryIndexAccess ScopeAccess(*this);
+			DirectoryLocation = FindLocationFromIndex(FullPath, MountPoint, DirectoryIndex);
+		}
+
+		if ((PathHashLocation != nullptr) != (DirectoryLocation != nullptr))
+		{
+			const TCHAR* FoundName = TEXT("PathHashIndex");
+			const TCHAR* NotFoundName = TEXT("FullDirectoryIndex");
+			if (!PathHashLocation)
+			{
+				Swap(FoundName, NotFoundName);
+			}
+			UE_LOG(LogPakFile, Error, TEXT("PathHashIndex does not match FullDirectoryIndex. Pakfile '%s' has '%s' in its %s but not in its %s."),
+				*PakFilename, *FullPath, FoundName, NotFoundName);
+		}
+		PakEntryLocation = PathHashLocation ? PathHashLocation : DirectoryLocation;
+	}
+	else
+#endif
+	{
+		if (bHasPathHashIndex)
+		{
+			PakEntryLocation = FindLocationFromIndex(FullPath, MountPoint, PathHashIndex, PathHashSeed, Info.Version);
+		}
+		else
+		{
+			check(bHasFullDirectoryIndex);
+			FScopedPakDirectoryIndexAccess ScopeAccess(*this);
+			PakEntryLocation = FindLocationFromIndex(FullPath, MountPoint, DirectoryIndex);
+		}
+	}
+	if (!PakEntryLocation)
+	{
+		return EFindResult::NotFound;
+	}
+
+	return GetPakEntry(*PakEntryLocation, OutEntry);
+}
+
+void FPakFile::AddSpecialFile(const FPakEntry& Entry, const FString& Filename)
+{
+	MakeDirectoryFromPath(MountPoint);
+
+	// TODO: This function is not threadsafe; readers of the Indexes will be invalidated when we modify them
+	// To make it threadsafe would require always holding the lock around any read of either index, which is
+	// more expensive than we want to support this debug feature
+	FPakEntryLocation EntryLocation;
+	if (!Entry.IsDeleteRecord())
+	{
+		// Add new file info.
+		TArray<uint8> NewEncodedPakEntries;
+		FMemoryWriter MemoryWriter(NewEncodedPakEntries);
+		EntryLocation = FPakEntryLocation::CreateFromOffsetIntoEncoded(EncodedPakEntries.Num());
+		if (EncodePakEntry(MemoryWriter, Entry, Info))
+		{
+			EncodedPakEntries.Append(NewEncodedPakEntries);
+			EncodedPakEntries.Shrink();
+		}
+		else
+		{
+			EntryLocation = FPakEntryLocation::CreateFromListIndex(Files.Num());
+			Files.Add(Entry);
+			Files.Shrink();
+		}
+		NumEntries++;
+	}
+
+	FPathHashIndex* PathHashToWrite = bHasPathHashIndex ? &PathHashIndex : nullptr;
+	AddEntryToIndex(Filename, EntryLocation, MountPoint, PathHashSeed, &DirectoryIndex, PathHashToWrite, nullptr /* CollisionDetection */, Info.Version);
+}
diff --git a/Engine/Source/Runtime/PakFile/Private/PakFile.inl b/Engine/Source/Runtime/PakFile/Private/PakFile.inl
new file mode 100644
--- /dev/null
+++ b/Engine/Source/Runtime/PakFile/Private/PakFile.inl
@@ -0,0 +1,123 @@
+// Copyright Epic Games, Inc. All Rights Reserved.
+
+#pragma once
+
+/**
+ * Search the given FDirectoryIndex for all files under the given Directory.  Helper for FindFilesAtPath, called separately on the DirectoryIndex or Pruned DirectoryIndex. Does not use
+ * FScopedPakDirectoryIndexAccess internally; caller is responsible for calling from within a lock.
+ * Returned paths are full paths (include the mount point)
+ */
+template <typename ShouldVisitFunc, class ContainerType>
+void FPakFile::FindFilesAtPathInIndex(const FDirectoryIndex& TargetIndex, ContainerType& OutFiles,
+	const FString& Directory, const ShouldVisitFunc& ShouldVisit, bool bIncludeFiles, bool bIncludeDirectories,
+	bool bRecursive) const
+{
+	// Early out if MountPoint is not matching directory
+	if (!Directory.StartsWith(MountPoint))
+	{
+		return;
+	}
+
+	FStringView RelativeSearch(FStringView(Directory).RightChop(MountPoint.Len()));
+
+	TArray<FString> DirectoriesInPak; // List of all unique directories at path
+	for (TMap<FString, FPakDirectory>::TConstIterator It(TargetIndex); It; ++It)
+	{
+		// Check if the file is under the specified path.
+		if (FStringView(It.Key()).StartsWith(RelativeSearch))
+		{
+			FString PakPath = PakPathCombine(MountPoint, It.Key());
+			if (bRecursive == true)
+			{
+				// Add everything
+				if (bIncludeFiles)
+				{
+					for (FPakDirectory::TConstIterator DirectoryIt(It.Value()); DirectoryIt; ++DirectoryIt)
+					{
+						const FString& FilePathUnderDirectory = DirectoryIt.Key();
+						if (ShouldVisit(FilePathUnderDirectory))
+						{
+							OutFiles.Add(PakPathCombine(PakPath, FilePathUnderDirectory));
+						}
+					}
+				}
+				if (bIncludeDirectories)
+				{
+					if (Directory != PakPath)
+					{
+						if (ShouldVisit(PakPath))
+						{
+							DirectoriesInPak.Add(MoveTemp(PakPath));
+						}
+					}
+				}
+			}
+			else
+			{
+				int32 SubDirIndex = PakPath.Len() > Directory.Len() ? PakPath.Find(TEXT("/"), ESearchCase::CaseSensitive, ESearchDir::FromStart, Directory.Len() + 1) : INDEX_NONE;
+				// Add files in the specified folder only.
+				if (bIncludeFiles && SubDirIndex == INDEX_NONE)
+				{
+					for (FPakDirectory::TConstIterator DirectoryIt(It.Value()); DirectoryIt; ++DirectoryIt)
+					{
+						const FString& FilePathUnderDirectory = DirectoryIt.Key();
+						if (ShouldVisit(FilePathUnderDirectory))
+						{
+							OutFiles.Add(PakPathCombine(PakPath, FilePathUnderDirectory));
+						}
+					}
+				}
+				// Add sub-folders in the specified folder only
+				if (bIncludeDirectories && SubDirIndex >= 0)
+				{
+					FString SubDirPath = PakPath.Left(SubDirIndex + 1);
+					if (ShouldVisit(SubDirPath))
+					{
+						DirectoriesInPak.AddUnique(MoveTemp(SubDirPath));
+					}
+				}
+			}
+		}
+	}
+	OutFiles.Append(MoveTemp(DirectoriesInPak));
+}
+
+template <typename ShouldVisitFunc, class ContainerType>
+void FPakFile::FindPrunedFilesAtPathInternal(const TCHAR* InPath, const ShouldVisitFunc& ShouldVisit, ContainerType& OutFiles,
+	bool bIncludeFiles, bool bIncludeDirectories, bool bRecursive) const
+{
+	// Make sure all directory names end with '/'.
+	FString Directory(InPath);
+	MakeDirectoryFromPath(Directory);
+
+	// Check the specified path is under the mount point of this pak file.
+	// The reverse case (MountPoint StartsWith Directory) is needed to properly handle
+	// pak files that are a subdirectory of the actual directory.
+	if (!Directory.StartsWith(MountPoint) && !MountPoint.StartsWith(Directory))
+	{
+		return;
+	}
+
+	FScopedPakDirectoryIndexAccess ScopeAccess(*this);
+#if ENABLE_PAKFILE_RUNTIME_PRUNING_VALIDATE
+	if (ShouldValidatePrunedDirectory())
+	{
+		TSet<FString> FullFoundFiles, PrunedFoundFiles;
+		FindFilesAtPathInIndex(DirectoryIndex, FullFoundFiles, Directory, ShouldVisit,
+			bIncludeFiles, bIncludeDirectories, bRecursive);
+		FindFilesAtPathInIndex(PrunedDirectoryIndex, PrunedFoundFiles, Directory, ShouldVisit,
+			bIncludeFiles, bIncludeDirectories, bRecursive);
+		ValidateDirectorySearch(FullFoundFiles, PrunedFoundFiles, InPath);
+
+		for (const FString& FoundFile : FullFoundFiles)
+		{
+			OutFiles.Add(FoundFile);
+		}
+	}
+	else
+#endif
+	{
+		FindFilesAtPathInIndex(DirectoryIndex, OutFiles, Directory, ShouldVisit,
+			bIncludeFiles, bIncludeDirectories, bRecursive);
+	}
+}
diff --git a/Engine/Source/Runtime/PakFile/Private/PakFileModule.cpp b/Engine/Source/Runtime/PakFile/Private/PakFileModule.cpp
new file mode 100644
--- /dev/null
+++ b/Engine/Source/Runtime/PakFile/Private/PakFileModule.cpp
@@ -0,0 +1,43 @@
+// Copyright Epic Games, Inc. All Rights Reserved.
+
+#include "HAL/IPlatformFileModule.h"
+#include "HAL/PlatformFileManager.h"
+#include "IPlatformFilePak.h"
+#include "Modules/ModuleManager.h"
+
+/**
+ * Module for the pak file
+ */
+class FPakFileModule : public IPlatformFileModule
+{
+public:
+	virtual IPlatformFile* GetPlatformFile() override
+	{
+		check(Singleton.IsValid());
+		return Singleton.Get();
+	}
+
+	virtual void StartupModule() override
+	{
+		Singleton = MakeUnique<FPakPlatformFile>();
+		FModuleManager::LoadModuleChecked<IModuleInterface>(TEXT("RSA"));
+	}
+
+	virtual void ShutdownModule() override
+	{
+		// remove ourselves from the platform file chain (there can be late writes after the shutdown).
+		if (Singleton.IsValid())
+		{
+			if (FPlatformFileManager::Get().FindPlatformFile(Singleton.Get()->GetName()))
+			{
+				FPlatformFileManager::Get().RemovePlatformFile(Singleton.Get());
+			}
+		}
+
+		Singleton.Reset();
+	}
+
+	TUniquePtr<IPlatformFile> Singleton;
+};
+
+IMPLEMENT_MODULE(FPakFileModule, PakFile);
diff --git a/Engine/Source/Runtime/PakFile/Private/PakIntervalTree.inl b/Engine/Source/Runtime/PakFile/Private/PakIntervalTree.inl
new file mode 100644
--- /dev/null
+++ b/Engine/Source/Runtime/PakFile/Private/PakIntervalTree.inl
@@ -0,0 +1,642 @@
+// Copyright Epic Games, Inc. All Rights Reserved.
+
+#pragma once
+
+enum
+{
+	IntervalTreeInvalidIndex = 0
+};
+
+typedef uint32 TIntervalTreeIndex; // this is the arg type of TSparseArray::operator[]
+
+typedef uint64 FJoinedOffsetAndPakIndex;
+
+static FORCEINLINE uint16 GetRequestPakIndexLow(FJoinedOffsetAndPakIndex Joined)
+{
+	return uint16((Joined >> 48) & 0xffff);
+}
+
+static FORCEINLINE int64 GetRequestOffset(FJoinedOffsetAndPakIndex Joined)
+{
+	return int64(Joined & 0xffffffffffffll);
+}
+
+static FORCEINLINE FJoinedOffsetAndPakIndex MakeJoinedRequest(uint16 PakIndex, int64 Offset)
+{
+	check(Offset >= 0);
+	return (FJoinedOffsetAndPakIndex(PakIndex) << 48) | Offset;
+}
+
+static uint32 GNextSalt = 1;
+
+// This is like TSparseArray, only a bit safer and I needed some restrictions on resizing.
+template<class TItem>
+class TIntervalTreeAllocator
+{
+	TArray<TItem> Items;
+	TArray<int32> FreeItems; //@todo make this into a linked list through the existing items
+	uint32 Salt;
+	uint32 SaltMask;
+public:
+	TIntervalTreeAllocator()
+	{
+		check(GNextSalt < 4);
+		Salt = (GNextSalt++) << 30;
+		SaltMask = MAX_uint32 << 30;
+		verify((Alloc() & ~SaltMask) == IntervalTreeInvalidIndex); // we want this to always have element zero so we can figure out an index from a pointer
+	}
+	inline TIntervalTreeIndex Alloc()
+	{
+		int32 Result;
+		if (FreeItems.Num())
+		{
+			Result = FreeItems.Pop();
+		}
+		else
+		{
+			Result = Items.Num();
+			Items.AddUninitialized();
+
+		}
+		new ((void*)&Items[Result]) TItem();
+		return Result | Salt;;
+	}
+	void EnsureNoRealloc(int32 NeededNewNum)
+	{
+		if (FreeItems.Num() + Items.GetSlack() < NeededNewNum)
+		{
+			Items.Reserve(Items.Num() + NeededNewNum);
+		}
+	}
+	FORCEINLINE TItem& Get(TIntervalTreeIndex InIndex)
+	{
+		TIntervalTreeIndex Index = InIndex & ~SaltMask;
+		check((InIndex & SaltMask) == Salt && Index != IntervalTreeInvalidIndex && Index >= 0 && Index < (uint32)Items.Num()); //&& !FreeItems.Contains(Index));
+		return Items[Index];
+	}
+	FORCEINLINE void Free(TIntervalTreeIndex InIndex)
+	{
+		TIntervalTreeIndex Index = InIndex & ~SaltMask;
+		check((InIndex & SaltMask) == Salt && Index != IntervalTreeInvalidIndex && Index >= 0 && Index < (uint32)Items.Num()); //&& !FreeItems.Contains(Index));
+		Items[Index].~TItem();
+		FreeItems.Push(Index);
+		if (FreeItems.Num() + 1 == Items.Num())
+		{
+			// get rid everything to restore memory coherence
+			Items.Empty();
+			FreeItems.Empty();
+			verify((Alloc() & ~SaltMask) == IntervalTreeInvalidIndex); // we want this to always have element zero so we can figure out an index from a pointer
+		}
+	}
+	FORCEINLINE void CheckIndex(TIntervalTreeIndex InIndex)
+	{
+		TIntervalTreeIndex Index = InIndex & ~SaltMask;
+		check((InIndex & SaltMask) == Salt && Index != IntervalTreeInvalidIndex && Index >= 0 && Index < (uint32)Items.Num()); // && !FreeItems.Contains(Index));
+	}
+};
+
+class FIntervalTreeNode
+{
+public:
+	TIntervalTreeIndex LeftChildOrRootOfLeftList;
+	TIntervalTreeIndex RootOfOnList;
+	TIntervalTreeIndex RightChildOrRootOfRightList;
+
+	FIntervalTreeNode()
+		: LeftChildOrRootOfLeftList(IntervalTreeInvalidIndex)
+		, RootOfOnList(IntervalTreeInvalidIndex)
+		, RightChildOrRootOfRightList(IntervalTreeInvalidIndex)
+	{
+	}
+	~FIntervalTreeNode()
+	{
+		check(LeftChildOrRootOfLeftList == IntervalTreeInvalidIndex && RootOfOnList == IntervalTreeInvalidIndex && RightChildOrRootOfRightList == IntervalTreeInvalidIndex); // this routine does not handle recursive destruction
+	}
+};
+
+static TIntervalTreeAllocator<FIntervalTreeNode> GIntervalTreeNodeNodeAllocator;
+
+static FORCEINLINE uint64 HighBit(uint64 x)
+{
+	return x & (1ull << 63);
+}
+
+static FORCEINLINE bool IntervalsIntersect(uint64 Min1, uint64 Max1, uint64 Min2, uint64 Max2)
+{
+	return !(Max2 < Min1 || Max1 < Min2);
+}
+
+template<typename TItem>
+// this routine assume that the pointers remain valid even though we are reallocating
+static void AddToIntervalTree_Dangerous(
+	TIntervalTreeIndex* RootNode,
+	TIntervalTreeAllocator<TItem>& Allocator,
+	TIntervalTreeIndex Index,
+	uint64 MinInterval,
+	uint64 MaxInterval,
+	uint32 CurrentShift,
+	uint32 MaxShift
+)
+{
+	while (true)
+	{
+		if (*RootNode == IntervalTreeInvalidIndex)
+		{
+			*RootNode = GIntervalTreeNodeNodeAllocator.Alloc();
+		}
+
+		int64 MinShifted = HighBit(MinInterval << CurrentShift);
+		int64 MaxShifted = HighBit(MaxInterval << CurrentShift);
+		FIntervalTreeNode& Root = GIntervalTreeNodeNodeAllocator.Get(*RootNode);
+
+		if (MinShifted == MaxShifted && CurrentShift < MaxShift)
+		{
+			CurrentShift++;
+			RootNode = (!MinShifted) ? &Root.LeftChildOrRootOfLeftList : &Root.RightChildOrRootOfRightList;
+		}
+		else
+		{
+			TItem& Item = Allocator.Get(Index);
+			if (MinShifted != MaxShifted) // crosses middle
+			{
+				Item.Next = Root.RootOfOnList;
+				Root.RootOfOnList = Index;
+			}
+			else // we are at the leaf
+			{
+				if (!MinShifted)
+				{
+					Item.Next = Root.LeftChildOrRootOfLeftList;
+					Root.LeftChildOrRootOfLeftList = Index;
+				}
+				else
+				{
+					Item.Next = Root.RightChildOrRootOfRightList;
+					Root.RightChildOrRootOfRightList = Index;
+				}
+			}
+			return;
+		}
+	}
+}
+
+template<typename TItem>
+static void AddToIntervalTree(
+	TIntervalTreeIndex* RootNode,
+	TIntervalTreeAllocator<TItem>& Allocator,
+	TIntervalTreeIndex Index,
+	uint32 StartShift,
+	uint32 MaxShift
+)
+{
+	GIntervalTreeNodeNodeAllocator.EnsureNoRealloc(1 + MaxShift - StartShift);
+	TItem& Item = Allocator.Get(Index);
+	check(Item.Next == IntervalTreeInvalidIndex);
+	uint64 MinInterval = GetRequestOffset(Item.OffsetAndPakIndex);
+	uint64 MaxInterval = MinInterval + Item.Size - 1;
+	AddToIntervalTree_Dangerous(RootNode, Allocator, Index, MinInterval, MaxInterval, StartShift, MaxShift);
+
+}
+
+template<typename TItem>
+static FORCEINLINE bool ScanNodeListForRemoval(
+	TIntervalTreeIndex* Iter,
+	TIntervalTreeAllocator<TItem>& Allocator,
+	TIntervalTreeIndex Index,
+	uint64 MinInterval,
+	uint64 MaxInterval
+)
+{
+	while (*Iter != IntervalTreeInvalidIndex)
+	{
+
+		TItem& Item = Allocator.Get(*Iter);
+		if (*Iter == Index)
+		{
+			*Iter = Item.Next;
+			Item.Next = IntervalTreeInvalidIndex;
+			return true;
+		}
+		Iter = &Item.Next;
+	}
+	return false;
+}
+
+template<typename TItem>
+static bool RemoveFromIntervalTree(
+	TIntervalTreeIndex* RootNode,
+	TIntervalTreeAllocator<TItem>& Allocator,
+	TIntervalTreeIndex Index,
+	uint64 MinInterval,
+	uint64 MaxInterval,
+	uint32 CurrentShift,
+	uint32 MaxShift
+)
+{
+	bool bResult = false;
+	if (*RootNode != IntervalTreeInvalidIndex)
+	{
+		int64 MinShifted = HighBit(MinInterval << CurrentShift);
+		int64 MaxShifted = HighBit(MaxInterval << CurrentShift);
+		FIntervalTreeNode& Root = GIntervalTreeNodeNodeAllocator.Get(*RootNode);
+
+		if (!MinShifted && !MaxShifted)
+		{
+			if (CurrentShift == MaxShift)
+			{
+				bResult = ScanNodeListForRemoval(&Root.LeftChildOrRootOfLeftList, Allocator, Index, MinInterval, MaxInterval);
+			}
+			else
+			{
+				bResult = RemoveFromIntervalTree(&Root.LeftChildOrRootOfLeftList, Allocator, Index, MinInterval, MaxInterval, CurrentShift + 1, MaxShift);
+			}
+		}
+		else if (!MinShifted && MaxShifted)
+		{
+			bResult = ScanNodeListForRemoval(&Root.RootOfOnList, Allocator, Index, MinInterval, MaxInterval);
+		}
+		else
+		{
+			if (CurrentShift == MaxShift)
+			{
+				bResult = ScanNodeListForRemoval(&Root.RightChildOrRootOfRightList, Allocator, Index, MinInterval, MaxInterval);
+			}
+			else
+			{
+				bResult = RemoveFromIntervalTree(&Root.RightChildOrRootOfRightList, Allocator, Index, MinInterval, MaxInterval, CurrentShift + 1, MaxShift);
+			}
+		}
+		if (bResult)
+		{
+			if (Root.LeftChildOrRootOfLeftList == IntervalTreeInvalidIndex && Root.RootOfOnList == IntervalTreeInvalidIndex && Root.RightChildOrRootOfRightList == IntervalTreeInvalidIndex)
+			{
+				check(&Root == &GIntervalTreeNodeNodeAllocator.Get(*RootNode));
+				GIntervalTreeNodeNodeAllocator.Free(*RootNode);
+				*RootNode = IntervalTreeInvalidIndex;
+			}
+		}
+	}
+	return bResult;
+}
+
+template<typename TItem>
+static bool RemoveFromIntervalTree(
+	TIntervalTreeIndex* RootNode,
+	TIntervalTreeAllocator<TItem>& Allocator,
+	TIntervalTreeIndex Index,
+	uint32 StartShift,
+	uint32 MaxShift
+)
+{
+	TItem& Item = Allocator.Get(Index);
+	uint64 MinInterval = GetRequestOffset(Item.OffsetAndPakIndex);
+	uint64 MaxInterval = MinInterval + Item.Size - 1;
+	return RemoveFromIntervalTree(RootNode, Allocator, Index, MinInterval, MaxInterval, StartShift, MaxShift);
+}
+
+template<typename TItem>
+static FORCEINLINE void ScanNodeListForRemovalFunc(
+	TIntervalTreeIndex* Iter,
+	TIntervalTreeAllocator<TItem>& Allocator,
+	uint64 MinInterval,
+	uint64 MaxInterval,
+	TFunctionRef<bool(TIntervalTreeIndex)> Func
+)
+{
+	while (*Iter != IntervalTreeInvalidIndex)
+	{
+		TItem& Item = Allocator.Get(*Iter);
+		uint64 Offset = uint64(GetRequestOffset(Item.OffsetAndPakIndex));
+		uint64 LastByte = Offset + uint64(Item.Size) - 1;
+
+		// save the value and then clear it.
+		TIntervalTreeIndex NextIndex = Item.Next;
+		if (IntervalsIntersect(MinInterval, MaxInterval, Offset, LastByte) && Func(*Iter))
+		{
+			*Iter = NextIndex; // this may have already be deleted, so cannot rely on the memory block
+		}
+		else
+		{
+			Iter = &Item.Next;
+		}
+	}
+}
+
+template<typename TItem>
+static void MaybeRemoveOverlappingNodesInIntervalTree(
+	TIntervalTreeIndex* RootNode,
+	TIntervalTreeAllocator<TItem>& Allocator,
+	uint64 MinInterval,
+	uint64 MaxInterval,
+	uint64 MinNode,
+	uint64 MaxNode,
+	uint32 CurrentShift,
+	uint32 MaxShift,
+	TFunctionRef<bool(TIntervalTreeIndex)> Func
+)
+{
+	if (*RootNode != IntervalTreeInvalidIndex)
+	{
+		int64 MinShifted = HighBit(MinInterval << CurrentShift);
+		int64 MaxShifted = HighBit(MaxInterval << CurrentShift);
+		FIntervalTreeNode& Root = GIntervalTreeNodeNodeAllocator.Get(*RootNode);
+		uint64 Center = (MinNode + MaxNode + 1) >> 1;
+
+		//UE_LOG(LogTemp, Warning, TEXT("Exploring Node %X [%d, %d] %d%d     interval %llX %llX    node interval %llX %llX   center %llX  "), *RootNode, CurrentShift, MaxShift, !!MinShifted, !!MaxShifted, MinInterval, MaxInterval, MinNode, MaxNode, Center);
+
+
+		if (!MinShifted)
+		{
+			if (CurrentShift == MaxShift)
+			{
+				//UE_LOG(LogTemp, Warning, TEXT("LeftBottom %X [%d, %d] %d%d"), *RootNode, CurrentShift, MaxShift, !!MinShifted, !!MaxShifted);
+				ScanNodeListForRemovalFunc(&Root.LeftChildOrRootOfLeftList, Allocator, MinInterval, MaxInterval, Func);
+			}
+			else
+			{
+				//UE_LOG(LogTemp, Warning, TEXT("LeftRecur %X [%d, %d] %d%d"), *RootNode, CurrentShift, MaxShift, !!MinShifted, !!MaxShifted);
+				MaybeRemoveOverlappingNodesInIntervalTree(&Root.LeftChildOrRootOfLeftList, Allocator, MinInterval, FMath::Min(MaxInterval, Center - 1), MinNode, Center - 1, CurrentShift + 1, MaxShift, Func);
+			}
+		}
+
+		//UE_LOG(LogTemp, Warning, TEXT("Center %X [%d, %d] %d%d"), *RootNode, CurrentShift, MaxShift, !!MinShifted, !!MaxShifted);
+		ScanNodeListForRemovalFunc(&Root.RootOfOnList, Allocator, MinInterval, MaxInterval, Func);
+
+		if (MaxShifted)
+		{
+			if (CurrentShift == MaxShift)
+			{
+				//UE_LOG(LogTemp, Warning, TEXT("RightBottom %X [%d, %d] %d%d"), *RootNode, CurrentShift, MaxShift, !!MinShifted, !!MaxShifted);
+				ScanNodeListForRemovalFunc(&Root.RightChildOrRootOfRightList, Allocator, MinInterval, MaxInterval, Func);
+			}
+			else
+			{
+				//UE_LOG(LogTemp, Warning, TEXT("RightRecur %X [%d, %d] %d%d"), *RootNode, CurrentShift, MaxShift, !!MinShifted, !!MaxShifted);
+				MaybeRemoveOverlappingNodesInIntervalTree(&Root.RightChildOrRootOfRightList, Allocator, FMath::Max(MinInterval, Center), MaxInterval, Center, MaxNode, CurrentShift + 1, MaxShift, Func);
+			}
+		}
+
+		//UE_LOG(LogTemp, Warning, TEXT("Done Exploring Node %X [%d, %d] %d%d     interval %llX %llX    node interval %llX %llX   center %llX  "), *RootNode, CurrentShift, MaxShift, !!MinShifted, !!MaxShifted, MinInterval, MaxInterval, MinNode, MaxNode, Center);
+
+		if (Root.LeftChildOrRootOfLeftList == IntervalTreeInvalidIndex && Root.RootOfOnList == IntervalTreeInvalidIndex && Root.RightChildOrRootOfRightList == IntervalTreeInvalidIndex)
+		{
+			check(&Root == &GIntervalTreeNodeNodeAllocator.Get(*RootNode));
+			GIntervalTreeNodeNodeAllocator.Free(*RootNode);
+			*RootNode = IntervalTreeInvalidIndex;
+		}
+	}
+}
+
+
+template<typename TItem>
+static FORCEINLINE bool ScanNodeList(
+	TIntervalTreeIndex Iter,
+	TIntervalTreeAllocator<TItem>& Allocator,
+	uint64 MinInterval,
+	uint64 MaxInterval,
+	TFunctionRef<bool(TIntervalTreeIndex)> Func
+)
+{
+	while (Iter != IntervalTreeInvalidIndex)
+	{
+		TItem& Item = Allocator.Get(Iter);
+		uint64 Offset = uint64(GetRequestOffset(Item.OffsetAndPakIndex));
+		uint64 LastByte = Offset + uint64(Item.Size) - 1;
+		if (IntervalsIntersect(MinInterval, MaxInterval, Offset, LastByte))
+		{
+			if (!Func(Iter))
+			{
+				return false;
+			}
+		}
+		Iter = Item.Next;
+	}
+	return true;
+}
+
+template<typename TItem>
+static bool OverlappingNodesInIntervalTree(
+	TIntervalTreeIndex RootNode,
+	TIntervalTreeAllocator<TItem>& Allocator,
+	uint64 MinInterval,
+	uint64 MaxInterval,
+	uint64 MinNode,
+	uint64 MaxNode,
+	uint32 CurrentShift,
+	uint32 MaxShift,
+	TFunctionRef<bool(TIntervalTreeIndex)> Func
+)
+{
+	if (RootNode != IntervalTreeInvalidIndex)
+	{
+		int64 MinShifted = HighBit(MinInterval << CurrentShift);
+		int64 MaxShifted = HighBit(MaxInterval << CurrentShift);
+		FIntervalTreeNode& Root = GIntervalTreeNodeNodeAllocator.Get(RootNode);
+		uint64 Center = (MinNode + MaxNode + 1) >> 1;
+
+		if (!MinShifted)
+		{
+			if (CurrentShift == MaxShift)
+			{
+				if (!ScanNodeList(Root.LeftChildOrRootOfLeftList, Allocator, MinInterval, MaxInterval, Func))
+				{
+					return false;
+				}
+			}
+			else
+			{
+				if (!OverlappingNodesInIntervalTree(Root.LeftChildOrRootOfLeftList, Allocator, MinInterval, FMath::Min(MaxInterval, Center - 1), MinNode, Center - 1, CurrentShift + 1, MaxShift, Func))
+				{
+					return false;
+				}
+			}
+		}
+		if (!ScanNodeList(Root.RootOfOnList, Allocator, MinInterval, MaxInterval, Func))
+		{
+			return false;
+		}
+		if (MaxShifted)
+		{
+			if (CurrentShift == MaxShift)
+			{
+				if (!ScanNodeList(Root.RightChildOrRootOfRightList, Allocator, MinInterval, MaxInterval, Func))
+				{
+					return false;
+				}
+			}
+			else
+			{
+				if (!OverlappingNodesInIntervalTree(Root.RightChildOrRootOfRightList, Allocator, FMath::Max(MinInterval, Center), MaxInterval, Center, MaxNode, CurrentShift + 1, MaxShift, Func))
+				{
+					return false;
+				}
+			}
+		}
+	}
+	return true;
+}
+
+template<typename TItem>
+static bool ScanNodeListWithShrinkingInterval(
+	TIntervalTreeIndex Iter,
+	TIntervalTreeAllocator<TItem>& Allocator,
+	uint64 MinInterval,
+	uint64& MaxInterval,
+	TFunctionRef<bool(TIntervalTreeIndex)> Func
+)
+{
+	while (Iter != IntervalTreeInvalidIndex)
+	{
+		TItem& Item = Allocator.Get(Iter);
+		uint64 Offset = uint64(GetRequestOffset(Item.OffsetAndPakIndex));
+		uint64 LastByte = Offset + uint64(Item.Size) - 1;
+		//UE_LOG(LogTemp, Warning, TEXT("Test Overlap %llu %llu %llu %llu"), MinInterval, MaxInterval, Offset, LastByte);
+		if (IntervalsIntersect(MinInterval, MaxInterval, Offset, LastByte))
+		{
+			//UE_LOG(LogTemp, Warning, TEXT("Overlap %llu %llu %llu %llu"), MinInterval, MaxInterval, Offset, LastByte);
+			if (!Func(Iter))
+			{
+				return false;
+			}
+		}
+		Iter = Item.Next;
+	}
+	return true;
+}
+
+template<typename TItem>
+static bool OverlappingNodesInIntervalTreeWithShrinkingInterval(
+	TIntervalTreeIndex RootNode,
+	TIntervalTreeAllocator<TItem>& Allocator,
+	uint64 MinInterval,
+	uint64& MaxInterval,
+	uint64 MinNode,
+	uint64 MaxNode,
+	uint32 CurrentShift,
+	uint32 MaxShift,
+	TFunctionRef<bool(TIntervalTreeIndex)> Func
+)
+{
+	if (RootNode != IntervalTreeInvalidIndex)
+	{
+
+		int64 MinShifted = HighBit(MinInterval << CurrentShift);
+		int64 MaxShifted = HighBit(FMath::Min(MaxInterval, MaxNode) << CurrentShift); // since MaxInterval is changing, we cannot clamp it during recursion.
+		FIntervalTreeNode& Root = GIntervalTreeNodeNodeAllocator.Get(RootNode);
+		uint64 Center = (MinNode + MaxNode + 1) >> 1;
+
+		if (!MinShifted)
+		{
+			if (CurrentShift == MaxShift)
+			{
+				if (!ScanNodeListWithShrinkingInterval(Root.LeftChildOrRootOfLeftList, Allocator, MinInterval, MaxInterval, Func))
+				{
+					return false;
+				}
+			}
+			else
+			{
+				if (!OverlappingNodesInIntervalTreeWithShrinkingInterval(Root.LeftChildOrRootOfLeftList, Allocator, MinInterval, MaxInterval, MinNode, Center - 1, CurrentShift + 1, MaxShift, Func)) // since MaxInterval is changing, we cannot clamp it during recursion.
+				{
+					return false;
+				}
+			}
+		}
+		if (!ScanNodeListWithShrinkingInterval(Root.RootOfOnList, Allocator, MinInterval, MaxInterval, Func))
+		{
+			return false;
+		}
+		MaxShifted = HighBit(FMath::Min(MaxInterval, MaxNode) << CurrentShift); // since MaxInterval is changing, we cannot clamp it during recursion.
+		if (MaxShifted)
+		{
+			if (CurrentShift == MaxShift)
+			{
+				if (!ScanNodeListWithShrinkingInterval(Root.RightChildOrRootOfRightList, Allocator, MinInterval, MaxInterval, Func))
+				{
+					return false;
+				}
+			}
+			else
+			{
+				if (!OverlappingNodesInIntervalTreeWithShrinkingInterval(Root.RightChildOrRootOfRightList, Allocator, FMath::Max(MinInterval, Center), MaxInterval, Center, MaxNode, CurrentShift + 1, MaxShift, Func))
+				{
+					return false;
+				}
+			}
+		}
+	}
+	return true;
+}
+
+
+template<typename TItem>
+static void MaskInterval(
+	TIntervalTreeIndex Index,
+	TIntervalTreeAllocator<TItem>& Allocator,
+	uint64 MinInterval,
+	uint64 MaxInterval,
+	uint32 BytesToBitsShift,
+	uint64* Bits
+)
+{
+	TItem& Item = Allocator.Get(Index);
+	uint64 Offset = uint64(GetRequestOffset(Item.OffsetAndPakIndex));
+	uint64 LastByte = Offset + uint64(Item.Size) - 1;
+	uint64 InterMinInterval = FMath::Max(MinInterval, Offset);
+	uint64 InterMaxInterval = FMath::Min(MaxInterval, LastByte);
+	if (InterMinInterval <= InterMaxInterval)
+	{
+		uint32 FirstBit = uint32((InterMinInterval - MinInterval) >> BytesToBitsShift);
+		uint32 LastBit = uint32((InterMaxInterval - MinInterval) >> BytesToBitsShift);
+		uint32 FirstQWord = FirstBit >> 6;
+		uint32 LastQWord = LastBit >> 6;
+		uint32 FirstBitQWord = FirstBit & 63;
+		uint32 LastBitQWord = LastBit & 63;
+		if (FirstQWord == LastQWord)
+		{
+			Bits[FirstQWord] |= ((MAX_uint64 << FirstBitQWord) & (MAX_uint64 >> (63 - LastBitQWord)));
+		}
+		else
+		{
+			Bits[FirstQWord] |= (MAX_uint64 << FirstBitQWord);
+			for (uint32 QWordIndex = FirstQWord + 1; QWordIndex < LastQWord; QWordIndex++)
+			{
+				Bits[QWordIndex] = MAX_uint64;
+			}
+			Bits[LastQWord] |= (MAX_uint64 >> (63 - LastBitQWord));
+		}
+	}
+}
+
+
+
+template<typename TItem>
+static void OverlappingNodesInIntervalTreeMask(
+	TIntervalTreeIndex RootNode,
+	TIntervalTreeAllocator<TItem>& Allocator,
+	uint64 MinInterval,
+	uint64 MaxInterval,
+	uint64 MinNode,
+	uint64 MaxNode,
+	uint32 CurrentShift,
+	uint32 MaxShift,
+	uint32 BytesToBitsShift,
+	uint64* Bits
+)
+{
+	OverlappingNodesInIntervalTree(
+		RootNode,
+		Allocator,
+		MinInterval,
+		MaxInterval,
+		MinNode,
+		MaxNode,
+		CurrentShift,
+		MaxShift,
+		[&Allocator, MinInterval, MaxInterval, BytesToBitsShift, Bits](TIntervalTreeIndex Index) -> bool
+	{
+		MaskInterval(Index, Allocator, MinInterval, MaxInterval, BytesToBitsShift, Bits);
+		return true;
+	}
+	);
+}
