Reunite all FPakPlatformFile methods

#rnx

diff --git a/Engine/Source/Runtime/PakFile/Private/IPlatformFilePak.cpp b/Engine/Source/Runtime/PakFile/Private/IPlatformFilePak.cpp
--- a/Engine/Source/Runtime/PakFile/Private/IPlatformFilePak.cpp
+++ b/Engine/Source/Runtime/PakFile/Private/IPlatformFilePak.cpp
@@ -222,257 +222,6 @@
 #define EXCLUDE_NONPAK_UE_EXTENSIONS 1	// Use .Build.cs file to disable this if the game relies on accessing loose files
 #endif
 
-FFilenameSecurityDelegate& FPakPlatformFile::GetFilenameSecurityDelegate()
-{
-	static FFilenameSecurityDelegate Delegate;
-	return Delegate;
-}
-
-FPakCustomEncryptionDelegate& FPakPlatformFile::GetPakCustomEncryptionDelegate()
-{
-	static FPakCustomEncryptionDelegate Delegate;
-	return Delegate;
-}
-
-FPakPlatformFile::FPakSigningFailureHandlerData& FPakPlatformFile::GetPakSigningFailureHandlerData()
-{
-	static FPakSigningFailureHandlerData Instance;
-	return Instance;
-}
-
-void FPakPlatformFile::BroadcastPakChunkSignatureCheckFailure(const FPakChunkSignatureCheckFailedData& InData)
-{
-	FPakSigningFailureHandlerData& HandlerData = GetPakSigningFailureHandlerData();
-	FScopeLock Lock(&HandlerData.GetLock());
-	HandlerData.GetPakChunkSignatureCheckFailedDelegate().Broadcast(InData);
-}
-
-void FPakPlatformFile::BroadcastPakPrincipalSignatureTableCheckFailure(const FString& InFilename)
-{
-	FPakSigningFailureHandlerData& HandlerData = GetPakSigningFailureHandlerData();
-	FScopeLock Lock(&HandlerData.GetLock());
-	HandlerData.GetPrincipalSignatureTableCheckFailedDelegate().Broadcast(InFilename);
-}
-
-PRAGMA_DISABLE_DEPRECATION_WARNINGS
-void FPakPlatformFile::BroadcastPakMasterSignatureTableCheckFailure(const FString& InFilename)
-{
-	return BroadcastPakPrincipalSignatureTableCheckFailure(InFilename);
-}
-PRAGMA_ENABLE_DEPRECATION_WARNINGS
-
-
-FPakSetIndexSettings& FPakPlatformFile::GetPakSetIndexSettingsDelegate()
-{
-	static FPakSetIndexSettings Delegate;
-	return Delegate;
-}
-
-void FPakPlatformFile::GetPrunedFilenamesInChunk(const FString& InPakFilename, const TArray<int32>& InChunkIDs, TArray<FString>& OutFileList)
-{
-	TArray<FPakListEntry> Paks;
-	GetMountedPaks(Paks);
-
-	for (const FPakListEntry& Pak : Paks)
-	{
-		if (Pak.PakFile && Pak.PakFile->GetFilename() == InPakFilename)
-		{
-			Pak.PakFile->GetPrunedFilenamesInChunk(InChunkIDs, OutFileList);
-			break;
-		}
-	}
-}
-
-void FPakPlatformFile::GetFilenamesFromIostoreByBlockIndex(const FString& InContainerName, const TArray<int32>& InBlockIndex, TArray<FString>& OutFileList)
-{
-	FPakPlatformFile* PakPlatformFile = static_cast<FPakPlatformFile*>(FPlatformFileManager::Get().FindPlatformFile(FPakPlatformFile::GetTypeName()));
-	if (!PakPlatformFile || !PakPlatformFile->IoDispatcherFileBackend.IsValid())
-	{
-		return;
-	}
-
-	const TMap<FGuid, FAES::FAESKey> Keys = UE::FEncryptionKeyManager::Get().GetAllKeys();
-
-	FScopeLock ScopedLock(&PakPlatformFile->PakListCritical);
-	for (const FPakListEntry& PakListEntry : PakPlatformFile->PakFiles)
-	{
-		if (FPaths::GetBaseFilename(PakListEntry.PakFile->PakFilename) == InContainerName)
-		{
-			TUniquePtr<FIoStoreReader> IoStoreReader(new FIoStoreReader());
-			FIoStatus Status = IoStoreReader->Initialize(*FPaths::ChangeExtension(PakListEntry.PakFile->PakFilename, TEXT("")),  Keys);
-			if (Status.IsOk())
-			{
-				IoStoreReader->GetFilenamesByBlockIndex(InBlockIndex, OutFileList);
-			}
-	
-			break;
-		}
-	}
-}
-
-bool FPakPlatformFile::DirectoryExistsInPrunedPakFiles(const TCHAR* Directory)
-{
-	FString StandardPath = Directory;
-	FPaths::MakeStandardFilename(StandardPath);
-
-	TArray<FPakListEntry> Paks;
-	GetMountedPaks(Paks);
-
-	// Check all pak files.
-	for (int32 PakIndex = 0; PakIndex < Paks.Num(); PakIndex++)
-	{
-		if (Paks[PakIndex].PakFile->DirectoryExistsInPruned(*StandardPath))
-		{
-			return true;
-		}
-	}
-	return false;
-}
-
-bool FPakPlatformFile::FindFileInPakFiles(TArray<FPakListEntry>& Paks, const TCHAR* Filename,
-	TRefCountPtr<FPakFile>* OutPakFile, FPakEntry* OutEntry)
-{
-	FString StandardFilename(Filename);
-	FPaths::MakeStandardFilename(StandardFilename);
-
-	TArray<const FPakListEntry*, TInlineAllocator<1>> PaksWithDeleteRecord;
-	bool bFoundOlderVersionOfDeleteRecordPak = false;
-
-	for (int32 PakIndex = 0; PakIndex < Paks.Num(); PakIndex++)
-	{
-		const FPakListEntry& PakEntry = Paks[PakIndex];
-		FPakFile* PakFile = PakEntry.PakFile.GetReference();
-		if (!PakFile)
-		{
-			continue;
-		}
-
-		if (PaksWithDeleteRecord.Num() > 0)
-		{
-			if (Algo::AnyOf(PaksWithDeleteRecord, [&PakEntry](const FPakListEntry* DeletedPakEntry)
-				{
-					return DeletedPakEntry->ReadOrder > PakEntry.ReadOrder &&
-						DeletedPakEntry->PakFile->PakchunkIndex == PakEntry.PakFile->PakchunkIndex;
-				}))
-			{
-				// Found a delete record in a higher priority patch level, and this is an earlier version of the same file.
-				// Don't search in the file.
-				bFoundOlderVersionOfDeleteRecordPak = true;
-				continue;
-			}
-		}
-
-		FPakFile::EFindResult FindResult = PakFile->Find(StandardFilename, OutEntry);
-		if (FindResult == FPakFile::EFindResult::Found)
-		{
-			if (OutPakFile != NULL)
-			{
-				*OutPakFile = PakFile;
-			}
-			UE_CLOG(!PaksWithDeleteRecord.IsEmpty(), LogPakFile, Verbose,
-				TEXT("Delete Record: Ignored delete record for %s - found it in %s instead (asset was moved or duplicated between chunks)"),
-				Filename, *PakFile->GetFilename());
-			return true;
-		}
-		else if (FindResult == FPakFile::EFindResult::FoundDeleted)
-		{
-			PaksWithDeleteRecord.Add(&PakEntry);
-			UE_LOG(LogPakFile, Verbose, TEXT("Delete Record: Found a delete record for %s in %s"),
-				Filename, *PakFile->GetFilename());
-		}
-	}
-
-	if (!PaksWithDeleteRecord.IsEmpty())
-	{
-		UE_CLOG(bFoundOlderVersionOfDeleteRecordPak, LogPakFile, Verbose,
-			TEXT("Delete Record: Accepted a delete record for %s"), Filename);
-		UE_CLOG(!bFoundOlderVersionOfDeleteRecordPak, LogPakFile, Warning,
-			TEXT("Delete Record: No lower priority pak files looking for %s. (maybe not downloaded?)"), Filename);
-	}
-	return false;
-}
-
-bool FPakPlatformFile::FindFileInPakFiles(const TCHAR* Filename, TRefCountPtr<FPakFile>* OutPakFile,
-	FPakEntry* OutEntry)
-{
-	TArray<FPakListEntry> Paks;
-	GetMountedPaks(Paks);
-
-	return FindFileInPakFiles(Paks, Filename, OutPakFile, OutEntry);
-}
-
-bool FPakPlatformFile::DirectoryExists(const TCHAR* Directory)
-{
-	// Check pak files first.
-	if (DirectoryExistsInPrunedPakFiles(Directory))
-	{
-		return true;
-	}
-	// Directory does not exist in any of the pak files, continue searching using inner platform file.
-	bool Result = LowerLevel->DirectoryExists(Directory);
-	return Result;
-}
-
-bool FPakPlatformFile::CreateDirectory(const TCHAR* Directory)
-{
-	// Directories can be created only under the normal path
-	return LowerLevel->CreateDirectory(Directory);
-}
-
-bool FPakPlatformFile::DeleteDirectory(const TCHAR* Directory)
-{
-	// Even if the same directory exists outside of pak files it will never
-	// get truly deleted from pak and will still be reported by Iterate functions.
-	// Fail in cases like this.
-	if (DirectoryExistsInPrunedPakFiles(Directory))
-	{
-		return false;
-	}
-	// Directory does not exist in pak files so it's safe to delete.
-	return LowerLevel->DeleteDirectory(Directory);
-}
-
-FFileStatData FPakPlatformFile::GetStatData(const TCHAR* FilenameOrDirectory)
-{
-	// Check pak files first.
-	FPakEntry FileEntry;
-	TRefCountPtr<FPakFile> PakFile;
-	if (FindFileInPakFiles(FilenameOrDirectory, &PakFile, &FileEntry))
-	{
-		return FFileStatData(
-			PakFile->GetTimestamp(),
-			PakFile->GetTimestamp(),
-			PakFile->GetTimestamp(),
-			(FileEntry.CompressionMethodIndex != 0) ? FileEntry.UncompressedSize : FileEntry.Size,
-			false,	// IsDirectory
-			true	// IsReadOnly
-		);
-	}
-
-	// Then check pak directories
-	if (DirectoryExistsInPrunedPakFiles(FilenameOrDirectory))
-	{
-		FDateTime DirectoryTimeStamp = FDateTime::MinValue();
-		return FFileStatData(
-			DirectoryTimeStamp,
-			DirectoryTimeStamp,
-			DirectoryTimeStamp,
-			-1,		// FileSize
-			true,	// IsDirectory
-			true	// IsReadOnly
-		);
-	}
-
-	// Fall back to lower level.
-	FFileStatData FileStatData;
-	if (IsNonPakFilenameAllowed(FilenameOrDirectory))
-	{
-		FileStatData = LowerLevel->GetStatData(FilenameOrDirectory);
-	}
-
-	return FileStatData;
-}
-
 namespace UE::PakFile::Private
 {
 
@@ -565,109 +314,6 @@
 
 }
 
-bool FPakPlatformFile::IterateDirectory(const TCHAR* Directory, IPlatformFile::FDirectoryVisitor& Visitor)
-{
-	return IterateDirectoryInternal(Directory, Visitor, false /* bRecursive */);
-}
-
-bool FPakPlatformFile::IterateDirectoryInternal(const TCHAR* Directory,
-	IPlatformFile::FDirectoryVisitor& Visitor, bool bRecursive)
-{
-	using namespace UE::PakFile::Private;
-
-	FPakFileDirectoryVisitor PakVisitor(Visitor);
-	TSet<FString> FilesVisitedInPak;
-	bool Result = IterateDirectoryInPakFiles(Directory, PakVisitor, bRecursive, FilesVisitedInPak);
-	if (Result && LowerLevel->DirectoryExists(Directory))
-	{
-		// Iterate inner filesystem but don't visit any files that were found in the Paks
-		FPreventDuplicatesVisitor PreventDuplicatesVisitor(Visitor, FilesVisitedInPak);
-		IPlatformFile::FDirectoryVisitor& LowerLevelVisitor(
-			// For performance, skip using PreventDuplicatedVisitor if there were no hits in pak
-			FilesVisitedInPak.Num() ? PreventDuplicatesVisitor : Visitor
-		);
-		if (bRecursive)
-		{
-			Result = LowerLevel->IterateDirectoryRecursively(Directory, LowerLevelVisitor);
-		}
-		else
-		{
-			Result = LowerLevel->IterateDirectory(Directory, LowerLevelVisitor);
-		}
-	}
-	return Result;
-}
-
-bool FPakPlatformFile::IterateDirectoryInPakFiles(const TCHAR* Directory,
-	UE::PakFile::Private::FPakFileDirectoryVisitorBase& Visitor, bool bRecursive, TSet<FString>& FilesVisitedInPak)
-{
-	bool Result = true;
-
-	TArray<FPakListEntry> Paks;
-	FString StandardDirectory = Directory;
-	FPaths::MakeStandardFilename(StandardDirectory);
-
-	bool bIsDownloadableDir =
-		(
-			FPaths::HasProjectPersistentDownloadDir() &&
-			StandardDirectory.StartsWith(FPaths::ProjectPersistentDownloadDir())
-			) ||
-		StandardDirectory.StartsWith(FPaths::CloudDir());
-
-	// don't look for in pak files for target-only locations
-	if (!bIsDownloadableDir)
-	{
-		GetMountedPaks(Paks);
-	}
-
-	// Iterate pak files first
-	FString NormalizationBuffer;
-	TSet<FString> FilesVisitedInThisPak;
-	auto ShouldVisit = [&Visitor](FStringView UnnormalizedPath)
-	{
-		FStringView NormalizedPath = UE::String::RemoveFromEnd(UnnormalizedPath, TEXTVIEW("/"));
-		return Visitor.ShouldVisitLeafPathname(FPathViews::GetCleanFilename(NormalizedPath));
-	};
-	for (int32 PakIndex = 0; PakIndex < Paks.Num(); PakIndex++)
-	{
-		FPakFile& PakFile = *Paks[PakIndex].PakFile;
-
-		const bool bIncludeFiles = true;
-		const bool bIncludeFolders = true;
-
-		FilesVisitedInThisPak.Reset();
-		PakFile.FindPrunedFilesAtPathInternal(*StandardDirectory, ShouldVisit, FilesVisitedInThisPak,
-			bIncludeFiles, bIncludeFolders, bRecursive);
-		for (TSet<FString>::TConstIterator SetIt(FilesVisitedInThisPak); SetIt && Result; ++SetIt)
-		{
-			const FString& Filename = *SetIt;
-			bool bIsDir = Filename.Len() && Filename[Filename.Len() - 1] == '/';
-			const FString* NormalizedFilename;
-			if (bIsDir)
-			{
-				NormalizationBuffer.Reset(Filename.Len());
-				NormalizationBuffer.AppendChars(*Filename, Filename.Len() - 1); // Chop off the trailing /
-				NormalizedFilename = &NormalizationBuffer;
-			}
-			else
-			{
-				NormalizedFilename = &Filename;
-			}
-			if (!FilesVisitedInPak.Contains(*NormalizedFilename))
-			{
-				FilesVisitedInPak.Add(*NormalizedFilename);
-				Result = Visitor.Visit(Filename, *NormalizedFilename, bIsDir, PakFile) && Result;
-			}
-		}
-	}
-	return Result;
-}
-
-bool FPakPlatformFile::IterateDirectoryRecursively(const TCHAR* Directory, IPlatformFile::FDirectoryVisitor& Visitor)
-{
-	return IterateDirectoryInternal(Directory, Visitor, true /* bRecursive */);
-}
-
 namespace UE::PakFile::Private
 {
 
@@ -737,215 +383,6 @@
 
 }
 
-bool FPakPlatformFile::IterateDirectoryStat(const TCHAR* Directory, IPlatformFile::FDirectoryStatVisitor& Visitor)
-{
-	return IterateDirectoryStatInternal(Directory, Visitor, false /* bRecursive */);
-}
-
-bool FPakPlatformFile::IterateDirectoryStatInternal(const TCHAR* Directory,
-	IPlatformFile::FDirectoryStatVisitor& Visitor, bool bRecursive)
-{
-	using namespace UE::PakFile::Private;
-
-	FPakFileDirectoryStatVisitor PakVisitor(*this, Visitor);
-	TSet<FString> FilesVisitedInPak;
-	bool Result = IterateDirectoryInPakFiles(Directory, PakVisitor, bRecursive, FilesVisitedInPak);
-	if (Result && LowerLevel->DirectoryExists(Directory))
-	{
-		// Iterate inner filesystem but don't visit any files that were found in the Paks
-		FPreventDuplicatesStatVisitor PreventDuplicatesVisitor(Visitor, FilesVisitedInPak);
-		IPlatformFile::FDirectoryStatVisitor& LowerLevelVisitor(
-			// For performance, skip using PreventDuplicatedVisitor if there were no hits in pak
-			FilesVisitedInPak.Num() ? PreventDuplicatesVisitor : Visitor);
-		if (bRecursive)
-		{
-			Result = LowerLevel->IterateDirectoryStatRecursively(Directory, LowerLevelVisitor);
-		}
-		else
-		{
-			Result = LowerLevel->IterateDirectoryStat(Directory, LowerLevelVisitor);
-		}
-	}
-	return Result;
-}
-
-bool FPakPlatformFile::IterateDirectoryStatRecursively(const TCHAR* Directory,
-	IPlatformFile::FDirectoryStatVisitor& Visitor)
-{
-	return IterateDirectoryStatInternal(Directory, Visitor, true/* bRecursive */);
-}
-
-void FPakPlatformFile::FindFiles(TArray<FString>& FoundFiles, const TCHAR* Directory, const TCHAR* FileExtension)
-{
-	if (LowerLevel->DirectoryExists(Directory))
-	{
-		LowerLevel->FindFiles(FoundFiles, Directory, FileExtension);
-	}
-
-	bool bRecursive = false;
-	FindFilesInternal(FoundFiles, Directory, FileExtension, bRecursive);
-}
-
-void FPakPlatformFile::FindFilesRecursively(TArray<FString>& FoundFiles,
-	const TCHAR* Directory, const TCHAR* FileExtension)
-{
-	if (LowerLevel->DirectoryExists(Directory))
-	{
-		LowerLevel->FindFilesRecursively(FoundFiles, Directory, FileExtension);
-	}
-
-	bool bRecursive = true;
-	FindFilesInternal(FoundFiles, Directory, FileExtension, bRecursive);
-}
-
-void FPakPlatformFile::FindFilesInternal(TArray<FString>& FoundFiles,
-	const TCHAR* Directory, const TCHAR* FileExtension, bool bRecursive)
-{
-	TArray<FPakListEntry> Paks;
-	GetMountedPaks(Paks);
-	if (Paks.Num())
-	{
-		TSet<FString> FilesVisited;
-		FilesVisited.Append(FoundFiles);
-
-		FString StandardDirectory = Directory;
-		FStringView FileExtensionStr = FileExtension;
-		FPaths::MakeStandardFilename(StandardDirectory);
-		bool bIncludeFiles = true;
-		bool bIncludeFolders = false;
-
-		auto ShouldVisit = [FileExtensionStr](FStringView Filename)
-		{
-			// filter out files by FileExtension
-			return FileExtensionStr.Len() == 0 || Filename.EndsWith(FileExtensionStr, ESearchCase::IgnoreCase);
-		};
-
-		TArray<FString> FilesInPak;
-		FilesInPak.Reserve(64);
-		for (int32 PakIndex = 0; PakIndex < Paks.Num(); PakIndex++)
-		{
-			FPakFile& PakFile = *Paks[PakIndex].PakFile;
-			PakFile.FindPrunedFilesAtPathInternal(*StandardDirectory, ShouldVisit, FilesInPak,
-				bIncludeFiles, bIncludeFolders, bRecursive);
-		}
-
-		for (const FString& Filename : FilesInPak)
-		{
-			// make sure we don't add duplicates to FoundFiles
-			bool bVisited = false;
-			FilesVisited.Add(Filename, &bVisited);
-			if (!bVisited)
-			{
-				FoundFiles.Add(Filename);
-			}
-		}
-	}
-}
-
-bool FPakPlatformFile::DeleteDirectoryRecursively(const TCHAR* Directory)
-{
-	// Can't delete directories existing in pak files. See DeleteDirectory(..) for more info.
-	if (DirectoryExistsInPrunedPakFiles(Directory))
-	{
-		return false;
-	}
-	// Directory does not exist in pak files so it's safe to delete.
-	return LowerLevel->DeleteDirectoryRecursively(Directory);
-}
-
-bool FPakPlatformFile::CreateDirectoryTree(const TCHAR* Directory)
-{
-	// Directories can only be created only under the normal path
-	return LowerLevel->CreateDirectoryTree(Directory);
-}
-
-void FPakPlatformFile::GetPrunedFilenamesInPakFile(const FString& InPakFilename, TArray<FString>& OutFileList)
-{
-	TArray<FPakListEntry> Paks;
-	GetMountedPaks(Paks);
-
-	for (const FPakListEntry& Pak : Paks)
-	{
-		if (Pak.PakFile && Pak.PakFile->GetFilename() == InPakFilename)
-		{
-			Pak.PakFile->GetPrunedFilenames(OutFileList);
-			break;
-		}
-	}
-}
-
-void FPakPlatformFile::GetFilenamesFromIostoreContainer(const FString& InContainerName, TArray<FString>& OutFileList)
-{
-	FPakPlatformFile* PakPlatformFile = static_cast<FPakPlatformFile*>(FPlatformFileManager::Get().FindPlatformFile(FPakPlatformFile::GetTypeName()));
-	if (!PakPlatformFile || !PakPlatformFile->IoDispatcherFileBackend.IsValid())
-	{
-		return;
-	}
-
-	const TMap<FGuid, FAES::FAESKey> Keys = UE::FEncryptionKeyManager::Get().GetAllKeys();
-
-	FScopeLock ScopedLock(&PakPlatformFile->PakListCritical);
-	for (const FPakListEntry& PakListEntry : PakPlatformFile->PakFiles)
-	{
-		if (FPaths::GetBaseFilename(PakListEntry.PakFile->PakFilename) == InContainerName)
-		{
-			TUniquePtr<FIoStoreReader> IoStoreReader(new FIoStoreReader());
-			FIoStatus Status = IoStoreReader->Initialize(*FPaths::ChangeExtension(PakListEntry.PakFile->PakFilename, TEXT("")), Keys);
-			if (Status.IsOk())
-			{
-				IoStoreReader->GetFilenames(OutFileList);
-			}
-			break;
-		}
-	}
-}
-
-void FPakPlatformFile::ForeachPackageInIostoreWhile(TFunctionRef<bool(FName)> Predicate)
-{
-	FPakPlatformFile* PakPlatformFile = static_cast<FPakPlatformFile*>(FPlatformFileManager::Get().FindPlatformFile(FPakPlatformFile::GetTypeName()));
-	if (!PakPlatformFile || !PakPlatformFile->IoDispatcherFileBackend.IsValid())
-	{
-		return;
-	}
-
-	const TMap<FGuid, FAES::FAESKey> Keys = UE::FEncryptionKeyManager::Get().GetAllKeys();
-
-	FScopeLock ScopedLock(&PakPlatformFile->PakListCritical);
-	for (const FPakListEntry& PakListEntry : PakPlatformFile->PakFiles)
-	{
-		TUniquePtr<FIoStoreReader> IoStoreReader(new FIoStoreReader());
-		FIoStatus Status = IoStoreReader->Initialize(*FPaths::ChangeExtension(PakListEntry.PakFile->PakFilename, TEXT("")), Keys);
-		if (Status.IsOk())
-		{
-			const FIoDirectoryIndexReader& DirectoryIndex = IoStoreReader->GetDirectoryIndexReader();
-
-			const bool Result = DirectoryIndex.IterateDirectoryIndex(
-				FIoDirectoryIndexHandle::RootDirectory(),
-				TEXT(""),
-				[Predicate](FStringView Filename, uint32) -> bool
-				{
-					const FStringView Ext = FPathViews::GetExtension(Filename);
-					if (Ext != TEXTVIEW("umap") && Ext != TEXTVIEW("uasset"))
-					{
-						return true; // ignore non package files
-					}
-
-					TStringBuilder<256> PackageNameBuilder;
-					if (FPackageName::TryConvertFilenameToLongPackageName(Filename, PackageNameBuilder))
-					{
-						return Invoke(Predicate, FName(PackageNameBuilder.ToView()));
-					}
-
-					return true; // ignore not mapped packages
-				});
-			if (!Result)
-			{
-				return;
-			}
-		}
-	}
-}
-
 #if !defined(PLATFORM_BYPASS_PAK_PRECACHE)
 	#error "PLATFORM_BYPASS_PAK_PRECACHE must be defined."
 #endif
@@ -956,94 +393,9 @@
 * Precaching
 */
 
-void FPakPlatformFile::GetPakEncryptionKey(FAES::FAESKey& OutKey, const FGuid& InEncryptionKeyGuid)
-{
-	OutKey.Reset();
-
-	if (!UE::FEncryptionKeyManager::Get().TryGetKey(InEncryptionKeyGuid, OutKey))
-	{
-		if (!InEncryptionKeyGuid.IsValid() && FCoreDelegates::GetPakEncryptionKeyDelegate().IsBound())
-		{
-			FCoreDelegates::GetPakEncryptionKeyDelegate().Execute(OutKey.Key);
-		}
-		else
-		{
-			UE_LOG(LogPakFile, Fatal, TEXT("Failed to find requested encryption key %s"), *InEncryptionKeyGuid.ToString());
-		}
-	}
-}
-
 TMap<FName, TSharedPtr<const FPakSignatureFile, ESPMode::ThreadSafe>> FPakPlatformFile::PakSignatureFileCache;
 FCriticalSection FPakPlatformFile::PakSignatureFileCacheLock;
 
-TSharedPtr<const FPakSignatureFile, ESPMode::ThreadSafe> FPakPlatformFile::GetPakSignatureFile(const TCHAR* InFilename)
-{
-	FName FilenameFName(InFilename);
-	{
-		FScopeLock Lock(&PakSignatureFileCacheLock);
-		if (TSharedPtr<const FPakSignatureFile, ESPMode::ThreadSafe>* ExistingSignatureFile = PakSignatureFileCache.Find(FilenameFName))
-		{
-			return *ExistingSignatureFile;
-		}
-	}
-
-	static FRSAKeyHandle PublicKey = []() -> FRSAKeyHandle
-	{
-		TDelegate<void(TArray<uint8>&, TArray<uint8>&)>& Delegate = FCoreDelegates::GetPakSigningKeysDelegate();
-		if (Delegate.IsBound())
-		{
-			TArray<uint8> Exponent;
-			TArray<uint8> Modulus;
-			Delegate.Execute(Exponent, Modulus);
-			return FRSA::CreateKey(Exponent, TArray<uint8>(), Modulus);
-		}
-		return InvalidRSAKeyHandle;
-	}();
-
-	TSharedPtr<FPakSignatureFile, ESPMode::ThreadSafe> NewSignatureFile;
-
-	if (PublicKey != InvalidRSAKeyHandle)
-	{
-		FString SignaturesFilename = FPaths::ChangeExtension(InFilename, TEXT("sig"));
-		TUniquePtr<FArchive> Reader(IFileManager::Get().CreateFileReader(*SignaturesFilename));
-		if (Reader != nullptr)
-		{
-			NewSignatureFile = MakeShared<FPakSignatureFile, ESPMode::ThreadSafe>();
-			NewSignatureFile->Serialize(*Reader);
-
-			if (!NewSignatureFile->DecryptSignatureAndValidate(PublicKey, InFilename))
-			{
-				// We don't need to act on this failure as the decrypt function will already have dumped out log messages
-				// and fired the signature check fail handler
-				NewSignatureFile.Reset();
-			}
-
-			{
-				FScopeLock Lock(&PakSignatureFileCacheLock);
-				if (TSharedPtr<const FPakSignatureFile, ESPMode::ThreadSafe>* ExistingSignatureFile = PakSignatureFileCache.Find(FilenameFName))
-				{
-					return *ExistingSignatureFile;
-				}
-				PakSignatureFileCache.Add(FilenameFName, NewSignatureFile);
-			}
-		}
-		else
-		{
-			UE_LOG(LogPakFile, Warning, TEXT("Couldn't find pak signature file '%s'"), InFilename);
-			BroadcastPakPrincipalSignatureTableCheckFailure(InFilename);
-		}
-	}
-
-	return NewSignatureFile;
-}
-
-void FPakPlatformFile::RemoveCachedPakSignaturesFile(const TCHAR* InFilename)
-{
-	FName FilenameFName(InFilename);
-	FScopeLock Lock(&PakSignatureFileCacheLock);
-	PakSignatureFileCache.Remove(FilenameFName);
-}
-
 DECLARE_DWORD_ACCUMULATOR_STAT(TEXT("PakCache Sync Decrypts (Uncompressed Path)"), STAT_PakCache_SyncDecrypts, STATGROUP_PakFile);
 DECLARE_FLOAT_ACCUMULATOR_STAT(TEXT("PakCache Decrypt Time"), STAT_PakCache_DecryptTime, STATGROUP_PakFile);
 DECLARE_DWORD_ACCUMULATOR_STAT(TEXT("PakCache Async Decrypts (Compressed Path)"), STAT_PakCache_CompressedDecrypts, STATGROUP_PakFile);
@@ -4484,16 +3836,6 @@
 
 #if PAK_TRACKER
 TMap<FString, int32> FPakPlatformFile::GPakSizeMap;
-
-void FPakPlatformFile::TrackPak(const TCHAR* Filename, const FPakEntry* PakEntry)
-{
-	FString Key(Filename);
-
-	if (!GPakSizeMap.Find(Key))
-	{
-		GPakSizeMap.Add(Key, PakEntry->Size);
-	}
-}
 #endif
 
 class FBypassPakAsyncReadFileHandle final : public IAsyncReadFileHandle
@@ -4560,108 +3902,6 @@
 	}
 };
 
-IAsyncReadFileHandle* FPakPlatformFile::OpenAsyncRead(const TCHAR* Filename)
-{
-	CSV_SCOPED_TIMING_STAT(FileIOVerbose, PakOpenAsyncRead);
-#if USE_PAK_PRECACHE
-	if (FPlatformProcess::SupportsMultithreading() && GPakCache_Enable > 0)
-	{
-		FPakEntry FileEntry;
-		TRefCountPtr<FPakFile> PakFile;
-		bool bFoundEntry = FindFileInPakFiles(Filename, &PakFile, &FileEntry);
-		if (bFoundEntry && PakFile && PakFile->GetFilenameName() != NAME_None)
-		{
-#if PAK_TRACKER
-			TrackPak(Filename, &FileEntry);
-#endif
-
-			return new FPakAsyncReadFileHandle(&FileEntry, PakFile, Filename);
-		}
-	}
-#elif PLATFORM_BYPASS_PAK_PRECACHE
-	{
-		FPakEntry FileEntry;
-		TRefCountPtr<FPakFile> PakFile;
-		bool bFoundEntry = FindFileInPakFiles(Filename, &PakFile, &FileEntry);
-		if (bFoundEntry && PakFile && PakFile->GetFilenameName() != NAME_None && FileEntry.CompressionMethodIndex == 0 && !FileEntry.IsEncrypted())
-		{
-#if PAK_TRACKER
-			TrackPak(Filename, &FileEntry);
-#endif
-			return new FBypassPakAsyncReadFileHandle(&FileEntry, PakFile, Filename);
-		}
-	}
-#endif
-	return IPlatformFile::OpenAsyncRead(Filename);
-}
-
-void FPakPlatformFile::SetAsyncMinimumPriority(EAsyncIOPriorityAndFlags Priority)
-{
-#if USE_PAK_PRECACHE
-	if (FPlatformProcess::SupportsMultithreading() && GPakCache_Enable > 0)
-	{
-		FPakPrecacher::Get().SetAsyncMinimumPriority(Priority);
-	}
-#elif PLATFORM_BYPASS_PAK_PRECACHE
-	IPlatformFile::GetPlatformPhysical().SetAsyncMinimumPriority(Priority);
-#endif
-}
-
-void FPakPlatformFile::Tick()
-{
-#if USE_PAK_PRECACHE && CSV_PROFILER_STATS
-	if (PakPrecacherSingleton != nullptr)
-	{
-		CSV_CUSTOM_STAT(FileIOVerbose, PakPrecacherRequests, FPakPrecacher::Get().GetRequestCount(), ECsvCustomStatOp::Set);
-		CSV_CUSTOM_STAT(FileIOVerbose, PakPrecacherHotBlocksCount, (int32)GPreCacheHotBlocksCount, ECsvCustomStatOp::Set);
-		CSV_CUSTOM_STAT(FileIOVerbose, PakPrecacherColdBlocksCount, (int32)GPreCacheColdBlocksCount, ECsvCustomStatOp::Set);
-		CSV_CUSTOM_STAT(FileIOVerbose, PakPrecacherTotalLoadedMB, (int32)(GPreCacheTotalLoaded / (1024 * 1024)), ECsvCustomStatOp::Set);
-		CSV_CUSTOM_STAT(FileIO, PakPrecacherBlockMemoryMB, (int32)(FPakPrecacher::Get().GetBlockMemory() / (1024 * 1024)), ECsvCustomStatOp::Set);
-
-
-
-		if (GPreCacheTotalLoadedLastTick != 0)
-		{
-			int64 diff = GPreCacheTotalLoaded - GPreCacheTotalLoadedLastTick;
-			diff /= 1024;
-			CSV_CUSTOM_STAT(FileIO, PakPrecacherPerFrameKB, (int32)diff, ECsvCustomStatOp::Set);
-		}
-		GPreCacheTotalLoadedLastTick = GPreCacheTotalLoaded;
-
-		CSV_CUSTOM_STAT(FileIOVerbose, PakPrecacherSeeks, (int32)GPreCacheSeeks, ECsvCustomStatOp::Set);
-		CSV_CUSTOM_STAT(FileIOVerbose, PakPrecacherBadSeeks, (int32)GPreCacheBadSeeks, ECsvCustomStatOp::Set);
-		CSV_CUSTOM_STAT(FileIOVerbose, PakPrecacherContiguousReads, (int32)GPreCacheContiguousReads, ECsvCustomStatOp::Set);
-		
-		CSV_CUSTOM_STAT(FileIOVerbose, PakLoads, (int32)PakPrecacherSingleton->Get().GetLoads(), ECsvCustomStatOp::Set);
-
-}
-#endif
-#if TRACK_DISK_UTILIZATION && CSV_PROFILER_STATS
-	CSV_CUSTOM_STAT(DiskIO, OutstandingIORequests, int32(GDiskUtilizationTracker.GetOutstandingRequests()), ECsvCustomStatOp::Set);
-	CSV_CUSTOM_STAT(DiskIO, BusyTime, float(GDiskUtilizationTracker.GetShortTermStats().GetTotalIOTimeInSeconds()), ECsvCustomStatOp::Set);
-	CSV_CUSTOM_STAT(DiskIO, IdleTime, float(GDiskUtilizationTracker.GetShortTermStats().GetTotalIdleTimeInSeconds()), ECsvCustomStatOp::Set);
-#endif
-
-#if CSV_PROFILER_STATS
-
-	int64 LocalTotalLoaded = GTotalLoaded;
-	if (IoDispatcherFileBackend.IsValid())
-	{
-		LocalTotalLoaded += FIoDispatcher::Get().GetTotalLoaded();
-	}
-
-	CSV_CUSTOM_STAT(FileIOVerbose, TotalLoadedMB, (int32)(LocalTotalLoaded / (1024 * 1024)), ECsvCustomStatOp::Set);
-	if (GTotalLoadedLastTick != 0)
-	{
-		int64 diff = LocalTotalLoaded - GTotalLoadedLastTick;
-		diff /= 1024;
-		CSV_CUSTOM_STAT(FileIO, PerFrameKB, (int32)diff, ECsvCustomStatOp::Set);
-	}
-	GTotalLoadedLastTick = LocalTotalLoaded;
-#endif
-
-}
-
 class FMappedFilePakProxy final : public IMappedFileHandle
 {
 	IMappedFileHandle* LowerLevel;
@@ -4735,53 +3975,6 @@
 	   );
 
 
-IMappedFileHandle* FPakPlatformFile::OpenMapped(const TCHAR* Filename)
-{
-	if (!GMMIO_Enable)
-	{
-		return nullptr;
-	}
-
-#if !UE_BUILD_SHIPPING
-	if (bLookLooseFirst && IsNonPakFilenameAllowed(Filename))
-	{
-		IMappedFileHandle* Handle = LowerLevel->OpenMapped(Filename);
-		if (Handle != nullptr)
-		{
-			return Handle;
-		}
-	}
-#endif
-
-	// Check pak files first
-	FPakEntry FileEntry;
-	TRefCountPtr<FPakFile> PakEntry;
-	if (FindFileInPakFiles(Filename, &PakEntry, &FileEntry) && PakEntry.IsValid())
-	{
-		if (FileEntry.CompressionMethodIndex != 0 || (FileEntry.Flags & FPakEntry::Flag_Encrypted) != 0)
-		{
-			// can't map compressed or encrypted files
-			return nullptr;
-		}
-		FScopeLock Lock(&PakEntry->MappedFileHandleCriticalSection);
-		if (!PakEntry->MappedFileHandle)
-		{
-			PakEntry->MappedFileHandle = LowerLevel->OpenMapped(*PakEntry->GetFilename());
-		}
-		if (!PakEntry->MappedFileHandle)
-		{
-			return nullptr;
-		}
-		return new FMappedFilePakProxy(PakEntry->MappedFileHandle, FileEntry.Offset + FileEntry.GetSerializedSize(PakEntry->GetInfo().Version), FileEntry.UncompressedSize, PakEntry->TotalSize(), Filename);
-	}
-	if (IsNonPakFilenameAllowed(Filename))
-	{
-		return LowerLevel->OpenMapped(Filename);
-	}
-	return nullptr;
-}
-
-
 /**
  * Class to handle correctly reading from a compressed file within a compressed package
  */
@@ -5139,94 +4332,8 @@
 	return bResult;
 }
 
-bool FPakPlatformFile::IsNonPakFilenameAllowed(const FString& InFilename)
-{
-	bool bAllowed = true;
-
-#if EXCLUDE_NONPAK_UE_EXTENSIONS
-	if (PakFiles.Num() || UE_BUILD_SHIPPING)
-	{
-		FName Ext = FName(*FPaths::GetExtension(InFilename));
-		bAllowed = !ExcludedNonPakExtensions.Contains(Ext);
-		UE_CLOG(!bAllowed, LogPakFile, VeryVerbose, TEXT("Access to file '%s' is limited to pak contents due to file extension being listed in ExcludedNonPakExtensions."), *InFilename)
-	}
-#endif
-
-	bool bIsIniFile = InFilename.EndsWith(IniFileExtension);
-#if DISABLE_NONUFS_INI_WHEN_COOKED
-	bool bSkipIniFile = bIsIniFile && !InFilename.EndsWith(GameUserSettingsIniFilename);
-	if (FPlatformProperties::RequiresCookedData() && bSkipIniFile)
-	{
-		bAllowed = false;
-	}
-#endif
-#if ALLOW_INI_OVERRIDE_FROM_COMMANDLINE
-	FString FileList;
-	if (bIsIniFile && FParse::Value(FCommandLine::Get(), TEXT("-iniFile="), FileList, false))
-	{
-		TArray<FString> Files;
-		FileList.ParseIntoArray(Files, TEXT(","), true);
-		for (int32 Index = 0; Index < Files.Num(); Index++)
-		{
-			if (InFilename == Files[Index])
-			{
-				bAllowed = true;
-				UE_LOG(LogPakFile, Log, TEXT(" Override -inifile: %s"), *InFilename);
-				break;
-			}
-		}
-	}
-#endif
-#if !DISABLE_CHEAT_CVARS && !UE_BUILD_SHIPPING
-	if (bIsIniFile && !bAllowed)
-	{
-		FString OverrideConsoleVariablesPath;
-		FParse::Value(FCommandLine::Get(), TEXT("-cvarsini="), OverrideConsoleVariablesPath);
-
-		if (!OverrideConsoleVariablesPath.IsEmpty() && InFilename == OverrideConsoleVariablesPath)
-		{
-			bAllowed = true;
-		}
-	}
-#endif
-
-	FFilenameSecurityDelegate& FilenameSecurityDelegate = GetFilenameSecurityDelegate();
-	if (bAllowed)
-	{
-		if (FilenameSecurityDelegate.IsBound())
-		{
-			bAllowed = FilenameSecurityDelegate.Execute(*InFilename);;
-		}
-	}
-
-	return bAllowed;
-}
-
-#if !HAS_PLATFORM_PAK_INSTALL_CHECK
-bool FPakPlatformFile::IsPakFileInstalled(const FString& InFilename)
-{
-#if ENABLE_PLATFORM_CHUNK_INSTALL
-	IPlatformChunkInstall* ChunkInstall = FPlatformMisc::GetPlatformChunkInstall();
-	if (ChunkInstall)
-	{
-		// if a platform supports chunk style installs, make sure that the chunk a pak file resides in is actually fully installed before accepting pak files from it
-		int32 PakchunkIndex = GetPakchunkIndexFromPakFile(InFilename);
-		if (PakchunkIndex != INDEX_NONE)
-		{
-			if (ChunkInstall->GetPakchunkLocation(PakchunkIndex) == EChunkLocation::NotAvailable)
-			{
-				return false;
-			}
-		}
-	}
-#endif
-
-	return true;
-}
-#endif //HAS_PLATFORM_PAK_INSTALL_CHECK
 
 
-
 FSharedPakReader::FSharedPakReader(FArchive* InArchive, FPakFile* InPakFile)
 	: Archive(InArchive)
 	, PakFile(InPakFile)
@@ -5361,6 +4468,8 @@
 }
 #endif // !UE_BUILD_SHIPPING
 
+
+
 FPakPlatformFile::FPakPlatformFile()
 	: LowerLevel(NULL)
 	, bSigned(false)
@@ -5394,4 +4503,896 @@
 	}
 }
 
+bool FPakPlatformFile::IsNonPakFilenameAllowed(const FString& InFilename)
+{
+	bool bAllowed = true;
+
+#if EXCLUDE_NONPAK_UE_EXTENSIONS
+	if (PakFiles.Num() || UE_BUILD_SHIPPING)
+	{
+		FName Ext = FName(*FPaths::GetExtension(InFilename));
+		bAllowed = !ExcludedNonPakExtensions.Contains(Ext);
+		UE_CLOG(!bAllowed, LogPakFile, VeryVerbose, TEXT("Access to file '%s' is limited to pak contents due to file extension being listed in ExcludedNonPakExtensions."), *InFilename)
+	}
+#endif
+
+	bool bIsIniFile = InFilename.EndsWith(IniFileExtension);
+#if DISABLE_NONUFS_INI_WHEN_COOKED
+	bool bSkipIniFile = bIsIniFile && !InFilename.EndsWith(GameUserSettingsIniFilename);
+	if (FPlatformProperties::RequiresCookedData() && bSkipIniFile)
+	{
+		bAllowed = false;
+	}
+#endif
+#if ALLOW_INI_OVERRIDE_FROM_COMMANDLINE
+	FString FileList;
+	if (bIsIniFile && FParse::Value(FCommandLine::Get(), TEXT("-iniFile="), FileList, false))
+	{
+		TArray<FString> Files;
+		FileList.ParseIntoArray(Files, TEXT(","), true);
+		for (int32 Index = 0; Index < Files.Num(); Index++)
+		{
+			if (InFilename == Files[Index])
+			{
+				bAllowed = true;
+				UE_LOG(LogPakFile, Log, TEXT(" Override -inifile: %s"), *InFilename);
+				break;
+			}
+		}
+	}
+#endif
+#if !DISABLE_CHEAT_CVARS && !UE_BUILD_SHIPPING
+	if (bIsIniFile && !bAllowed)
+	{
+		FString OverrideConsoleVariablesPath;
+		FParse::Value(FCommandLine::Get(), TEXT("-cvarsini="), OverrideConsoleVariablesPath);
+
+		if (!OverrideConsoleVariablesPath.IsEmpty() && InFilename == OverrideConsoleVariablesPath)
+		{
+			bAllowed = true;
+		}
+	}
+#endif
+
+	FFilenameSecurityDelegate& FilenameSecurityDelegate = GetFilenameSecurityDelegate();
+	if (bAllowed)
+	{
+		if (FilenameSecurityDelegate.IsBound())
+		{
+			bAllowed = FilenameSecurityDelegate.Execute(*InFilename);;
+		}
+	}
+
+	return bAllowed;
+}
+
+#if !HAS_PLATFORM_PAK_INSTALL_CHECK
+bool FPakPlatformFile::IsPakFileInstalled(const FString& InFilename)
+{
+#if ENABLE_PLATFORM_CHUNK_INSTALL
+	IPlatformChunkInstall* ChunkInstall = FPlatformMisc::GetPlatformChunkInstall();
+	if (ChunkInstall)
+	{
+		// if a platform supports chunk style installs, make sure that the chunk a pak file resides in is actually fully installed before accepting pak files from it
+		int32 PakchunkIndex = GetPakchunkIndexFromPakFile(InFilename);
+		if (PakchunkIndex != INDEX_NONE)
+		{
+			if (ChunkInstall->GetPakchunkLocation(PakchunkIndex) == EChunkLocation::NotAvailable)
+			{
+				return false;
+			}
+		}
+	}
+#endif
+
+	return true;
+}
+#endif //HAS_PLATFORM_PAK_INSTALL_CHECK
+
+IMappedFileHandle* FPakPlatformFile::OpenMapped(const TCHAR* Filename)
+{
+	if (!GMMIO_Enable)
+	{
+		return nullptr;
+	}
+
+#if !UE_BUILD_SHIPPING
+	if (bLookLooseFirst && IsNonPakFilenameAllowed(Filename))
+	{
+		IMappedFileHandle* Handle = LowerLevel->OpenMapped(Filename);
+		if (Handle != nullptr)
+		{
+			return Handle;
+		}
+	}
+#endif
+
+	// Check pak files first
+	FPakEntry FileEntry;
+	TRefCountPtr<FPakFile> PakEntry;
+	if (FindFileInPakFiles(Filename, &PakEntry, &FileEntry) && PakEntry.IsValid())
+	{
+		if (FileEntry.CompressionMethodIndex != 0 || (FileEntry.Flags & FPakEntry::Flag_Encrypted) != 0)
+		{
+			// can't map compressed or encrypted files
+			return nullptr;
+		}
+		FScopeLock Lock(&PakEntry->MappedFileHandleCriticalSection);
+		if (!PakEntry->MappedFileHandle)
+		{
+			PakEntry->MappedFileHandle = LowerLevel->OpenMapped(*PakEntry->GetFilename());
+		}
+		if (!PakEntry->MappedFileHandle)
+		{
+			return nullptr;
+		}
+		return new FMappedFilePakProxy(PakEntry->MappedFileHandle, FileEntry.Offset + FileEntry.GetSerializedSize(PakEntry->GetInfo().Version), FileEntry.UncompressedSize, PakEntry->TotalSize(), Filename);
+	}
+	if (IsNonPakFilenameAllowed(Filename))
+	{
+		return LowerLevel->OpenMapped(Filename);
+	}
+	return nullptr;
+}
+
+IAsyncReadFileHandle* FPakPlatformFile::OpenAsyncRead(const TCHAR* Filename)
+{
+	CSV_SCOPED_TIMING_STAT(FileIOVerbose, PakOpenAsyncRead);
+#if USE_PAK_PRECACHE
+	if (FPlatformProcess::SupportsMultithreading() && GPakCache_Enable > 0)
+	{
+		FPakEntry FileEntry;
+		TRefCountPtr<FPakFile> PakFile;
+		bool bFoundEntry = FindFileInPakFiles(Filename, &PakFile, &FileEntry);
+		if (bFoundEntry && PakFile && PakFile->GetFilenameName() != NAME_None)
+		{
+#if PAK_TRACKER
+			TrackPak(Filename, &FileEntry);
+#endif
+
+			return new FPakAsyncReadFileHandle(&FileEntry, PakFile, Filename);
+		}
+	}
+#elif PLATFORM_BYPASS_PAK_PRECACHE
+	{
+		FPakEntry FileEntry;
+		TRefCountPtr<FPakFile> PakFile;
+		bool bFoundEntry = FindFileInPakFiles(Filename, &PakFile, &FileEntry);
+		if (bFoundEntry && PakFile && PakFile->GetFilenameName() != NAME_None && FileEntry.CompressionMethodIndex == 0 && !FileEntry.IsEncrypted())
+		{
+#if PAK_TRACKER
+			TrackPak(Filename, &FileEntry);
+#endif
+			return new FBypassPakAsyncReadFileHandle(&FileEntry, PakFile, Filename);
+		}
+	}
+#endif
+	return IPlatformFile::OpenAsyncRead(Filename);
+}
+
+void FPakPlatformFile::SetAsyncMinimumPriority(EAsyncIOPriorityAndFlags Priority)
+{
+#if USE_PAK_PRECACHE
+	if (FPlatformProcess::SupportsMultithreading() && GPakCache_Enable > 0)
+	{
+		FPakPrecacher::Get().SetAsyncMinimumPriority(Priority);
+	}
+#elif PLATFORM_BYPASS_PAK_PRECACHE
+	IPlatformFile::GetPlatformPhysical().SetAsyncMinimumPriority(Priority);
+#endif
+}
+
+void FPakPlatformFile::Tick()
+{
+#if USE_PAK_PRECACHE && CSV_PROFILER_STATS
+	if (PakPrecacherSingleton != nullptr)
+	{
+		CSV_CUSTOM_STAT(FileIOVerbose, PakPrecacherRequests, FPakPrecacher::Get().GetRequestCount(), ECsvCustomStatOp::Set);
+		CSV_CUSTOM_STAT(FileIOVerbose, PakPrecacherHotBlocksCount, (int32)GPreCacheHotBlocksCount, ECsvCustomStatOp::Set);
+		CSV_CUSTOM_STAT(FileIOVerbose, PakPrecacherColdBlocksCount, (int32)GPreCacheColdBlocksCount, ECsvCustomStatOp::Set);
+		CSV_CUSTOM_STAT(FileIOVerbose, PakPrecacherTotalLoadedMB, (int32)(GPreCacheTotalLoaded / (1024 * 1024)), ECsvCustomStatOp::Set);
+		CSV_CUSTOM_STAT(FileIO, PakPrecacherBlockMemoryMB, (int32)(FPakPrecacher::Get().GetBlockMemory() / (1024 * 1024)), ECsvCustomStatOp::Set);
+
+		if (GPreCacheTotalLoadedLastTick != 0)
+		{
+			int64 diff = GPreCacheTotalLoaded - GPreCacheTotalLoadedLastTick;
+			diff /= 1024;
+			CSV_CUSTOM_STAT(FileIO, PakPrecacherPerFrameKB, (int32)diff, ECsvCustomStatOp::Set);
+		}
+		GPreCacheTotalLoadedLastTick = GPreCacheTotalLoaded;
+
+		CSV_CUSTOM_STAT(FileIOVerbose, PakPrecacherSeeks, (int32)GPreCacheSeeks, ECsvCustomStatOp::Set);
+		CSV_CUSTOM_STAT(FileIOVerbose, PakPrecacherBadSeeks, (int32)GPreCacheBadSeeks, ECsvCustomStatOp::Set);
+		CSV_CUSTOM_STAT(FileIOVerbose, PakPrecacherContiguousReads, (int32)GPreCacheContiguousReads, ECsvCustomStatOp::Set);
+		
+		CSV_CUSTOM_STAT(FileIOVerbose, PakLoads, (int32)PakPrecacherSingleton->Get().GetLoads(), ECsvCustomStatOp::Set);
+
+}
+#endif
+#if TRACK_DISK_UTILIZATION && CSV_PROFILER_STATS
+	CSV_CUSTOM_STAT(DiskIO, OutstandingIORequests, int32(GDiskUtilizationTracker.GetOutstandingRequests()), ECsvCustomStatOp::Set);
+	CSV_CUSTOM_STAT(DiskIO, BusyTime, float(GDiskUtilizationTracker.GetShortTermStats().GetTotalIOTimeInSeconds()), ECsvCustomStatOp::Set);
+	CSV_CUSTOM_STAT(DiskIO, IdleTime, float(GDiskUtilizationTracker.GetShortTermStats().GetTotalIdleTimeInSeconds()), ECsvCustomStatOp::Set);
+#endif
+
+#if CSV_PROFILER_STATS
+
+	int64 LocalTotalLoaded = GTotalLoaded;
+	if (IoDispatcherFileBackend.IsValid())
+	{
+		LocalTotalLoaded += FIoDispatcher::Get().GetTotalLoaded();
+	}
+
+	CSV_CUSTOM_STAT(FileIOVerbose, TotalLoadedMB, (int32)(LocalTotalLoaded / (1024 * 1024)), ECsvCustomStatOp::Set);
+	if (GTotalLoadedLastTick != 0)
+	{
+		int64 diff = LocalTotalLoaded - GTotalLoadedLastTick;
+		diff /= 1024;
+		CSV_CUSTOM_STAT(FileIO, PerFrameKB, (int32)diff, ECsvCustomStatOp::Set);
+	}
+	GTotalLoadedLastTick = LocalTotalLoaded;
+#endif
+
+}
+
+#if PAK_TRACKER
+void FPakPlatformFile::TrackPak(const TCHAR* Filename, const FPakEntry* PakEntry)
+{
+	FString Key(Filename);
+
+	if (!GPakSizeMap.Find(Key))
+	{
+		GPakSizeMap.Add(Key, PakEntry->Size);
+	}
+}
+#endif
+
+TSharedPtr<const FPakSignatureFile, ESPMode::ThreadSafe> FPakPlatformFile::GetPakSignatureFile(const TCHAR* InFilename)
+{
+	FName FilenameFName(InFilename);
+	{
+		FScopeLock Lock(&PakSignatureFileCacheLock);
+		if (TSharedPtr<const FPakSignatureFile, ESPMode::ThreadSafe>* ExistingSignatureFile = PakSignatureFileCache.Find(FilenameFName))
+		{
+			return *ExistingSignatureFile;
+		}
+	}
+
+	static FRSAKeyHandle PublicKey = []() -> FRSAKeyHandle
+	{
+		TDelegate<void(TArray<uint8>&, TArray<uint8>&)>& Delegate = FCoreDelegates::GetPakSigningKeysDelegate();
+		if (Delegate.IsBound())
+		{
+			TArray<uint8> Exponent;
+			TArray<uint8> Modulus;
+			Delegate.Execute(Exponent, Modulus);
+			return FRSA::CreateKey(Exponent, TArray<uint8>(), Modulus);
+		}
+		return InvalidRSAKeyHandle;
+	}();
+
+	TSharedPtr<FPakSignatureFile, ESPMode::ThreadSafe> NewSignatureFile;
+
+	if (PublicKey != InvalidRSAKeyHandle)
+	{
+		FString SignaturesFilename = FPaths::ChangeExtension(InFilename, TEXT("sig"));
+		TUniquePtr<FArchive> Reader(IFileManager::Get().CreateFileReader(*SignaturesFilename));
+		if (Reader != nullptr)
+		{
+			NewSignatureFile = MakeShared<FPakSignatureFile, ESPMode::ThreadSafe>();
+			NewSignatureFile->Serialize(*Reader);
+
+			if (!NewSignatureFile->DecryptSignatureAndValidate(PublicKey, InFilename))
+			{
+				// We don't need to act on this failure as the decrypt function will already have dumped out log messages
+				// and fired the signature check fail handler
+				NewSignatureFile.Reset();
+			}
+
+			{
+				FScopeLock Lock(&PakSignatureFileCacheLock);
+				if (TSharedPtr<const FPakSignatureFile, ESPMode::ThreadSafe>* ExistingSignatureFile = PakSignatureFileCache.Find(FilenameFName))
+				{
+					return *ExistingSignatureFile;
+				}
+				PakSignatureFileCache.Add(FilenameFName, NewSignatureFile);
+			}
+		}
+		else
+		{
+			UE_LOG(LogPakFile, Warning, TEXT("Couldn't find pak signature file '%s'"), InFilename);
+			BroadcastPakPrincipalSignatureTableCheckFailure(InFilename);
+		}
+	}
+
+	return NewSignatureFile;
+}
+
+void FPakPlatformFile::RemoveCachedPakSignaturesFile(const TCHAR* InFilename)
+{
+	FName FilenameFName(InFilename);
+	FScopeLock Lock(&PakSignatureFileCacheLock);
+	PakSignatureFileCache.Remove(FilenameFName);
+}
+
+void FPakPlatformFile::GetPakEncryptionKey(FAES::FAESKey& OutKey, const FGuid& InEncryptionKeyGuid)
+{
+	OutKey.Reset();
+
+	if (!UE::FEncryptionKeyManager::Get().TryGetKey(InEncryptionKeyGuid, OutKey))
+	{
+		if (!InEncryptionKeyGuid.IsValid() && FCoreDelegates::GetPakEncryptionKeyDelegate().IsBound())
+		{
+			FCoreDelegates::GetPakEncryptionKeyDelegate().Execute(OutKey.Key);
+		}
+		else
+		{
+			UE_LOG(LogPakFile, Fatal, TEXT("Failed to find requested encryption key %s"), *InEncryptionKeyGuid.ToString());
+		}
+	}
+}
+
+bool FPakPlatformFile::IterateDirectoryStat(const TCHAR* Directory, IPlatformFile::FDirectoryStatVisitor& Visitor)
+{
+	return IterateDirectoryStatInternal(Directory, Visitor, false /* bRecursive */);
+}
+
+bool FPakPlatformFile::IterateDirectoryStatInternal(const TCHAR* Directory,
+	IPlatformFile::FDirectoryStatVisitor& Visitor, bool bRecursive)
+{
+	using namespace UE::PakFile::Private;
+
+	FPakFileDirectoryStatVisitor PakVisitor(*this, Visitor);
+	TSet<FString> FilesVisitedInPak;
+	bool Result = IterateDirectoryInPakFiles(Directory, PakVisitor, bRecursive, FilesVisitedInPak);
+	if (Result && LowerLevel->DirectoryExists(Directory))
+	{
+		// Iterate inner filesystem but don't visit any files that were found in the Paks
+		FPreventDuplicatesStatVisitor PreventDuplicatesVisitor(Visitor, FilesVisitedInPak);
+		IPlatformFile::FDirectoryStatVisitor& LowerLevelVisitor(
+			// For performance, skip using PreventDuplicatedVisitor if there were no hits in pak
+			FilesVisitedInPak.Num() ? PreventDuplicatesVisitor : Visitor);
+		if (bRecursive)
+		{
+			Result = LowerLevel->IterateDirectoryStatRecursively(Directory, LowerLevelVisitor);
+		}
+		else
+		{
+			Result = LowerLevel->IterateDirectoryStat(Directory, LowerLevelVisitor);
+		}
+	}
+	return Result;
+}
+
+bool FPakPlatformFile::IterateDirectoryStatRecursively(const TCHAR* Directory,
+	IPlatformFile::FDirectoryStatVisitor& Visitor)
+{
+	return IterateDirectoryStatInternal(Directory, Visitor, true/* bRecursive */);
+}
+
+void FPakPlatformFile::FindFiles(TArray<FString>& FoundFiles, const TCHAR* Directory, const TCHAR* FileExtension)
+{
+	if (LowerLevel->DirectoryExists(Directory))
+	{
+		LowerLevel->FindFiles(FoundFiles, Directory, FileExtension);
+	}
+
+	bool bRecursive = false;
+	FindFilesInternal(FoundFiles, Directory, FileExtension, bRecursive);
+}
+
+void FPakPlatformFile::FindFilesRecursively(TArray<FString>& FoundFiles,
+	const TCHAR* Directory, const TCHAR* FileExtension)
+{
+	if (LowerLevel->DirectoryExists(Directory))
+	{
+		LowerLevel->FindFilesRecursively(FoundFiles, Directory, FileExtension);
+	}
+
+	bool bRecursive = true;
+	FindFilesInternal(FoundFiles, Directory, FileExtension, bRecursive);
+}
+
+void FPakPlatformFile::FindFilesInternal(TArray<FString>& FoundFiles,
+	const TCHAR* Directory, const TCHAR* FileExtension, bool bRecursive)
+{
+	TArray<FPakListEntry> Paks;
+	GetMountedPaks(Paks);
+	if (Paks.Num())
+	{
+		TSet<FString> FilesVisited;
+		FilesVisited.Append(FoundFiles);
+
+		FString StandardDirectory = Directory;
+		FStringView FileExtensionStr = FileExtension;
+		FPaths::MakeStandardFilename(StandardDirectory);
+		bool bIncludeFiles = true;
+		bool bIncludeFolders = false;
+
+		auto ShouldVisit = [FileExtensionStr](FStringView Filename)
+		{
+			// filter out files by FileExtension
+			return FileExtensionStr.Len() == 0 || Filename.EndsWith(FileExtensionStr, ESearchCase::IgnoreCase);
+		};
+
+		TArray<FString> FilesInPak;
+		FilesInPak.Reserve(64);
+		for (int32 PakIndex = 0; PakIndex < Paks.Num(); PakIndex++)
+		{
+			FPakFile& PakFile = *Paks[PakIndex].PakFile;
+			PakFile.FindPrunedFilesAtPathInternal(*StandardDirectory, ShouldVisit, FilesInPak,
+				bIncludeFiles, bIncludeFolders, bRecursive);
+		}
+
+		for (const FString& Filename : FilesInPak)
+		{
+			// make sure we don't add duplicates to FoundFiles
+			bool bVisited = false;
+			FilesVisited.Add(Filename, &bVisited);
+			if (!bVisited)
+			{
+				FoundFiles.Add(Filename);
+			}
+		}
+	}
+}
+
+bool FPakPlatformFile::DeleteDirectoryRecursively(const TCHAR* Directory)
+{
+	// Can't delete directories existing in pak files. See DeleteDirectory(..) for more info.
+	if (DirectoryExistsInPrunedPakFiles(Directory))
+	{
+		return false;
+	}
+	// Directory does not exist in pak files so it's safe to delete.
+	return LowerLevel->DeleteDirectoryRecursively(Directory);
+}
+
+bool FPakPlatformFile::CreateDirectoryTree(const TCHAR* Directory)
+{
+	// Directories can only be created only under the normal path
+	return LowerLevel->CreateDirectoryTree(Directory);
+}
+
+void FPakPlatformFile::GetPrunedFilenamesInPakFile(const FString& InPakFilename, TArray<FString>& OutFileList)
+{
+	TArray<FPakListEntry> Paks;
+	GetMountedPaks(Paks);
+
+	for (const FPakListEntry& Pak : Paks)
+	{
+		if (Pak.PakFile && Pak.PakFile->GetFilename() == InPakFilename)
+		{
+			Pak.PakFile->GetPrunedFilenames(OutFileList);
+			break;
+		}
+	}
+}
+
+void FPakPlatformFile::GetFilenamesFromIostoreContainer(const FString& InContainerName, TArray<FString>& OutFileList)
+{
+	FPakPlatformFile* PakPlatformFile = static_cast<FPakPlatformFile*>(FPlatformFileManager::Get().FindPlatformFile(FPakPlatformFile::GetTypeName()));
+	if (!PakPlatformFile || !PakPlatformFile->IoDispatcherFileBackend.IsValid())
+	{
+		return;
+	}
+
+	const TMap<FGuid, FAES::FAESKey> Keys = UE::FEncryptionKeyManager::Get().GetAllKeys();
+
+	FScopeLock ScopedLock(&PakPlatformFile->PakListCritical);
+	for (const FPakListEntry& PakListEntry : PakPlatformFile->PakFiles)
+	{
+		if (FPaths::GetBaseFilename(PakListEntry.PakFile->PakFilename) == InContainerName)
+		{
+			TUniquePtr<FIoStoreReader> IoStoreReader(new FIoStoreReader());
+			FIoStatus Status = IoStoreReader->Initialize(*FPaths::ChangeExtension(PakListEntry.PakFile->PakFilename, TEXT("")), Keys);
+			if (Status.IsOk())
+			{
+				IoStoreReader->GetFilenames(OutFileList);
+			}
+			break;
+		}
+	}
+}
+
+void FPakPlatformFile::ForeachPackageInIostoreWhile(TFunctionRef<bool(FName)> Predicate)
+{
+	FPakPlatformFile* PakPlatformFile = static_cast<FPakPlatformFile*>(FPlatformFileManager::Get().FindPlatformFile(FPakPlatformFile::GetTypeName()));
+	if (!PakPlatformFile || !PakPlatformFile->IoDispatcherFileBackend.IsValid())
+	{
+		return;
+	}
+
+	const TMap<FGuid, FAES::FAESKey> Keys = UE::FEncryptionKeyManager::Get().GetAllKeys();
+
+	FScopeLock ScopedLock(&PakPlatformFile->PakListCritical);
+	for (const FPakListEntry& PakListEntry : PakPlatformFile->PakFiles)
+	{
+		TUniquePtr<FIoStoreReader> IoStoreReader(new FIoStoreReader());
+		FIoStatus Status = IoStoreReader->Initialize(*FPaths::ChangeExtension(PakListEntry.PakFile->PakFilename, TEXT("")), Keys);
+		if (Status.IsOk())
+		{
+			const FIoDirectoryIndexReader& DirectoryIndex = IoStoreReader->GetDirectoryIndexReader();
+
+			const bool Result = DirectoryIndex.IterateDirectoryIndex(
+				FIoDirectoryIndexHandle::RootDirectory(),
+				TEXT(""),
+				[Predicate](FStringView Filename, uint32) -> bool
+				{
+					const FStringView Ext = FPathViews::GetExtension(Filename);
+					if (Ext != TEXTVIEW("umap") && Ext != TEXTVIEW("uasset"))
+					{
+						return true; // ignore non package files
+					}
+
+					TStringBuilder<256> PackageNameBuilder;
+					if (FPackageName::TryConvertFilenameToLongPackageName(Filename, PackageNameBuilder))
+					{
+						return Invoke(Predicate, FName(PackageNameBuilder.ToView()));
+					}
+
+					return true; // ignore not mapped packages
+				});
+			if (!Result)
+			{
+				return;
+			}
+		}
+	}
+}
+
+bool FPakPlatformFile::IterateDirectory(const TCHAR* Directory, IPlatformFile::FDirectoryVisitor& Visitor)
+{
+	return IterateDirectoryInternal(Directory, Visitor, false /* bRecursive */);
+}
+
+bool FPakPlatformFile::IterateDirectoryInternal(const TCHAR* Directory,
+	IPlatformFile::FDirectoryVisitor& Visitor, bool bRecursive)
+{
+	using namespace UE::PakFile::Private;
+
+	FPakFileDirectoryVisitor PakVisitor(Visitor);
+	TSet<FString> FilesVisitedInPak;
+	bool Result = IterateDirectoryInPakFiles(Directory, PakVisitor, bRecursive, FilesVisitedInPak);
+	if (Result && LowerLevel->DirectoryExists(Directory))
+	{
+		// Iterate inner filesystem but don't visit any files that were found in the Paks
+		FPreventDuplicatesVisitor PreventDuplicatesVisitor(Visitor, FilesVisitedInPak);
+		IPlatformFile::FDirectoryVisitor& LowerLevelVisitor(
+			// For performance, skip using PreventDuplicatedVisitor if there were no hits in pak
+			FilesVisitedInPak.Num() ? PreventDuplicatesVisitor : Visitor
+		);
+		if (bRecursive)
+		{
+			Result = LowerLevel->IterateDirectoryRecursively(Directory, LowerLevelVisitor);
+		}
+		else
+		{
+			Result = LowerLevel->IterateDirectory(Directory, LowerLevelVisitor);
+		}
+	}
+	return Result;
+}
+
+bool FPakPlatformFile::IterateDirectoryInPakFiles(const TCHAR* Directory,
+	UE::PakFile::Private::FPakFileDirectoryVisitorBase& Visitor, bool bRecursive, TSet<FString>& FilesVisitedInPak)
+{
+	bool Result = true;
+
+	TArray<FPakListEntry> Paks;
+	FString StandardDirectory = Directory;
+	FPaths::MakeStandardFilename(StandardDirectory);
+
+	bool bIsDownloadableDir =
+		(
+			FPaths::HasProjectPersistentDownloadDir() &&
+			StandardDirectory.StartsWith(FPaths::ProjectPersistentDownloadDir())
+			) ||
+		StandardDirectory.StartsWith(FPaths::CloudDir());
+
+	// don't look for in pak files for target-only locations
+	if (!bIsDownloadableDir)
+	{
+		GetMountedPaks(Paks);
+	}
+
+	// Iterate pak files first
+	FString NormalizationBuffer;
+	TSet<FString> FilesVisitedInThisPak;
+	auto ShouldVisit = [&Visitor](FStringView UnnormalizedPath)
+	{
+		FStringView NormalizedPath = UE::String::RemoveFromEnd(UnnormalizedPath, TEXTVIEW("/"));
+		return Visitor.ShouldVisitLeafPathname(FPathViews::GetCleanFilename(NormalizedPath));
+	};
+	for (int32 PakIndex = 0; PakIndex < Paks.Num(); PakIndex++)
+	{
+		FPakFile& PakFile = *Paks[PakIndex].PakFile;
+
+		const bool bIncludeFiles = true;
+		const bool bIncludeFolders = true;
+
+		FilesVisitedInThisPak.Reset();
+		PakFile.FindPrunedFilesAtPathInternal(*StandardDirectory, ShouldVisit, FilesVisitedInThisPak,
+			bIncludeFiles, bIncludeFolders, bRecursive);
+		for (TSet<FString>::TConstIterator SetIt(FilesVisitedInThisPak); SetIt && Result; ++SetIt)
+		{
+			const FString& Filename = *SetIt;
+			bool bIsDir = Filename.Len() && Filename[Filename.Len() - 1] == '/';
+			const FString* NormalizedFilename;
+			if (bIsDir)
+			{
+				NormalizationBuffer.Reset(Filename.Len());
+				NormalizationBuffer.AppendChars(*Filename, Filename.Len() - 1); // Chop off the trailing /
+				NormalizedFilename = &NormalizationBuffer;
+			}
+			else
+			{
+				NormalizedFilename = &Filename;
+			}
+			if (!FilesVisitedInPak.Contains(*NormalizedFilename))
+			{
+				FilesVisitedInPak.Add(*NormalizedFilename);
+				Result = Visitor.Visit(Filename, *NormalizedFilename, bIsDir, PakFile) && Result;
+			}
+		}
+	}
+	return Result;
+}
+
+bool FPakPlatformFile::IterateDirectoryRecursively(const TCHAR* Directory, IPlatformFile::FDirectoryVisitor& Visitor)
+{
+	return IterateDirectoryInternal(Directory, Visitor, true /* bRecursive */);
+}
+
+FFilenameSecurityDelegate& FPakPlatformFile::GetFilenameSecurityDelegate()
+{
+	static FFilenameSecurityDelegate Delegate;
+	return Delegate;
+}
+
+FPakCustomEncryptionDelegate& FPakPlatformFile::GetPakCustomEncryptionDelegate()
+{
+	static FPakCustomEncryptionDelegate Delegate;
+	return Delegate;
+}
+
+FPakPlatformFile::FPakSigningFailureHandlerData& FPakPlatformFile::GetPakSigningFailureHandlerData()
+{
+	static FPakSigningFailureHandlerData Instance;
+	return Instance;
+}
+
+void FPakPlatformFile::BroadcastPakChunkSignatureCheckFailure(const FPakChunkSignatureCheckFailedData& InData)
+{
+	FPakSigningFailureHandlerData& HandlerData = GetPakSigningFailureHandlerData();
+	FScopeLock Lock(&HandlerData.GetLock());
+	HandlerData.GetPakChunkSignatureCheckFailedDelegate().Broadcast(InData);
+}
+
+void FPakPlatformFile::BroadcastPakPrincipalSignatureTableCheckFailure(const FString& InFilename)
+{
+	FPakSigningFailureHandlerData& HandlerData = GetPakSigningFailureHandlerData();
+	FScopeLock Lock(&HandlerData.GetLock());
+	HandlerData.GetPrincipalSignatureTableCheckFailedDelegate().Broadcast(InFilename);
+}
+
+PRAGMA_DISABLE_DEPRECATION_WARNINGS
+void FPakPlatformFile::BroadcastPakMasterSignatureTableCheckFailure(const FString& InFilename)
+{
+	return BroadcastPakPrincipalSignatureTableCheckFailure(InFilename);
+}
+PRAGMA_ENABLE_DEPRECATION_WARNINGS
+
+
+FPakSetIndexSettings& FPakPlatformFile::GetPakSetIndexSettingsDelegate()
+{
+	static FPakSetIndexSettings Delegate;
+	return Delegate;
+}
+
+void FPakPlatformFile::GetPrunedFilenamesInChunk(const FString& InPakFilename, const TArray<int32>& InChunkIDs, TArray<FString>& OutFileList)
+{
+	TArray<FPakListEntry> Paks;
+	GetMountedPaks(Paks);
+
+	for (const FPakListEntry& Pak : Paks)
+	{
+		if (Pak.PakFile && Pak.PakFile->GetFilename() == InPakFilename)
+		{
+			Pak.PakFile->GetPrunedFilenamesInChunk(InChunkIDs, OutFileList);
+			break;
+		}
+	}
+}
+
+void FPakPlatformFile::GetFilenamesFromIostoreByBlockIndex(const FString& InContainerName, const TArray<int32>& InBlockIndex, TArray<FString>& OutFileList)
+{
+	FPakPlatformFile* PakPlatformFile = static_cast<FPakPlatformFile*>(FPlatformFileManager::Get().FindPlatformFile(FPakPlatformFile::GetTypeName()));
+	if (!PakPlatformFile || !PakPlatformFile->IoDispatcherFileBackend.IsValid())
+	{
+		return;
+	}
+
+	const TMap<FGuid, FAES::FAESKey> Keys = UE::FEncryptionKeyManager::Get().GetAllKeys();
+
+	FScopeLock ScopedLock(&PakPlatformFile->PakListCritical);
+	for (const FPakListEntry& PakListEntry : PakPlatformFile->PakFiles)
+	{
+		if (FPaths::GetBaseFilename(PakListEntry.PakFile->PakFilename) == InContainerName)
+		{
+			TUniquePtr<FIoStoreReader> IoStoreReader(new FIoStoreReader());
+			FIoStatus Status = IoStoreReader->Initialize(*FPaths::ChangeExtension(PakListEntry.PakFile->PakFilename, TEXT("")),  Keys);
+			if (Status.IsOk())
+			{
+				IoStoreReader->GetFilenamesByBlockIndex(InBlockIndex, OutFileList);
+			}
+	
+			break;
+		}
+	}
+}
+
+bool FPakPlatformFile::DirectoryExistsInPrunedPakFiles(const TCHAR* Directory)
+{
+	FString StandardPath = Directory;
+	FPaths::MakeStandardFilename(StandardPath);
+
+	TArray<FPakListEntry> Paks;
+	GetMountedPaks(Paks);
+
+	// Check all pak files.
+	for (int32 PakIndex = 0; PakIndex < Paks.Num(); PakIndex++)
+	{
+		if (Paks[PakIndex].PakFile->DirectoryExistsInPruned(*StandardPath))
+		{
+			return true;
+		}
+	}
+	return false;
+}
+
+bool FPakPlatformFile::FindFileInPakFiles(TArray<FPakListEntry>& Paks, const TCHAR* Filename,
+	TRefCountPtr<FPakFile>* OutPakFile, FPakEntry* OutEntry)
+{
+	FString StandardFilename(Filename);
+	FPaths::MakeStandardFilename(StandardFilename);
+
+	TArray<const FPakListEntry*, TInlineAllocator<1>> PaksWithDeleteRecord;
+	bool bFoundOlderVersionOfDeleteRecordPak = false;
+
+	for (int32 PakIndex = 0; PakIndex < Paks.Num(); PakIndex++)
+	{
+		const FPakListEntry& PakEntry = Paks[PakIndex];
+		FPakFile* PakFile = PakEntry.PakFile.GetReference();
+		if (!PakFile)
+		{
+			continue;
+		}
+
+		if (PaksWithDeleteRecord.Num() > 0)
+		{
+			if (Algo::AnyOf(PaksWithDeleteRecord, [&PakEntry](const FPakListEntry* DeletedPakEntry)
+				{
+					return DeletedPakEntry->ReadOrder > PakEntry.ReadOrder &&
+						DeletedPakEntry->PakFile->PakchunkIndex == PakEntry.PakFile->PakchunkIndex;
+				}))
+			{
+				// Found a delete record in a higher priority patch level, and this is an earlier version of the same file.
+				// Don't search in the file.
+				bFoundOlderVersionOfDeleteRecordPak = true;
+				continue;
+			}
+		}
+
+		FPakFile::EFindResult FindResult = PakFile->Find(StandardFilename, OutEntry);
+		if (FindResult == FPakFile::EFindResult::Found)
+		{
+			if (OutPakFile != NULL)
+			{
+				*OutPakFile = PakFile;
+			}
+			UE_CLOG(!PaksWithDeleteRecord.IsEmpty(), LogPakFile, Verbose,
+				TEXT("Delete Record: Ignored delete record for %s - found it in %s instead (asset was moved or duplicated between chunks)"),
+				Filename, *PakFile->GetFilename());
+			return true;
+		}
+		else if (FindResult == FPakFile::EFindResult::FoundDeleted)
+		{
+			PaksWithDeleteRecord.Add(&PakEntry);
+			UE_LOG(LogPakFile, Verbose, TEXT("Delete Record: Found a delete record for %s in %s"),
+				Filename, *PakFile->GetFilename());
+		}
+	}
+
+	if (!PaksWithDeleteRecord.IsEmpty())
+	{
+		UE_CLOG(bFoundOlderVersionOfDeleteRecordPak, LogPakFile, Verbose,
+			TEXT("Delete Record: Accepted a delete record for %s"), Filename);
+		UE_CLOG(!bFoundOlderVersionOfDeleteRecordPak, LogPakFile, Warning,
+			TEXT("Delete Record: No lower priority pak files looking for %s. (maybe not downloaded?)"), Filename);
+	}
+	return false;
+}
+
+bool FPakPlatformFile::FindFileInPakFiles(const TCHAR* Filename, TRefCountPtr<FPakFile>* OutPakFile,
+	FPakEntry* OutEntry)
+{
+	TArray<FPakListEntry> Paks;
+	GetMountedPaks(Paks);
+
+	return FindFileInPakFiles(Paks, Filename, OutPakFile, OutEntry);
+}
+
+bool FPakPlatformFile::DirectoryExists(const TCHAR* Directory)
+{
+	// Check pak files first.
+	if (DirectoryExistsInPrunedPakFiles(Directory))
+	{
+		return true;
+	}
+	// Directory does not exist in any of the pak files, continue searching using inner platform file.
+	bool Result = LowerLevel->DirectoryExists(Directory);
+	return Result;
+}
+
+bool FPakPlatformFile::CreateDirectory(const TCHAR* Directory)
+{
+	// Directories can be created only under the normal path
+	return LowerLevel->CreateDirectory(Directory);
+}
+
+bool FPakPlatformFile::DeleteDirectory(const TCHAR* Directory)
+{
+	// Even if the same directory exists outside of pak files it will never
+	// get truly deleted from pak and will still be reported by Iterate functions.
+	// Fail in cases like this.
+	if (DirectoryExistsInPrunedPakFiles(Directory))
+	{
+		return false;
+	}
+	// Directory does not exist in pak files so it's safe to delete.
+	return LowerLevel->DeleteDirectory(Directory);
+}
+
+FFileStatData FPakPlatformFile::GetStatData(const TCHAR* FilenameOrDirectory)
+{
+	// Check pak files first.
+	FPakEntry FileEntry;
+	TRefCountPtr<FPakFile> PakFile;
+	if (FindFileInPakFiles(FilenameOrDirectory, &PakFile, &FileEntry))
+	{
+		return FFileStatData(
+			PakFile->GetTimestamp(),
+			PakFile->GetTimestamp(),
+			PakFile->GetTimestamp(),
+			(FileEntry.CompressionMethodIndex != 0) ? FileEntry.UncompressedSize : FileEntry.Size,
+			false,	// IsDirectory
+			true	// IsReadOnly
+		);
+	}
+
+	// Then check pak directories
+	if (DirectoryExistsInPrunedPakFiles(FilenameOrDirectory))
+	{
+		FDateTime DirectoryTimeStamp = FDateTime::MinValue();
+		return FFileStatData(
+			DirectoryTimeStamp,
+			DirectoryTimeStamp,
+			DirectoryTimeStamp,
+			-1,		// FileSize
+			true,	// IsDirectory
+			true	// IsReadOnly
+		);
+	}
+
+	// Fall back to lower level.
+	FFileStatData FileStatData;
+	if (IsNonPakFilenameAllowed(FilenameOrDirectory))
+	{
+		FileStatData = LowerLevel->GetStatData(FilenameOrDirectory);
+	}
+
+	return FileStatData;
+}
+
 void FPakPlatformFile::FindPakFilesInDirectory(IPlatformFile* LowLevelFile, const TCHAR* Directory, const FString& WildCard, TArray<FString>& OutPakFiles)
