Enable processing cancelled file read requests before starting I/O (s.IoDispatcherCullCancelledReadRequests=1)

#rnx

diff --git a/Engine/Source/Runtime/PakFile/Private/GenericPlatformIoDispatcher.cpp b/Engine/Source/Runtime/PakFile/Private/GenericPlatformIoDispatcher.cpp
--- a/Engine/Source/Runtime/PakFile/Private/GenericPlatformIoDispatcher.cpp
+++ b/Engine/Source/Runtime/PakFile/Private/GenericPlatformIoDispatcher.cpp
@@ -11,7 +11,7 @@
 
 //PRAGMA_DISABLE_OPTIMIZATION
 
-bool GIoDispatcherCullCancelledReadRequests = false;
+bool GIoDispatcherCullCancelledReadRequests = true;
 static FAutoConsoleVariableRef CVar_IoDispatcherCullCancelledReadRequests(
 	TEXT("s.IoDispatcherCullCancelledReadRequests"),
 	GIoDispatcherCullCancelledReadRequests,
