Notify I/O service thread when cancelling read requests

#rnx

diff --git a/Engine/Source/Runtime/PakFile/Private/IoDispatcherFileBackend.cpp b/Engine/Source/Runtime/PakFile/Private/IoDispatcherFileBackend.cpp
--- a/Engine/Source/Runtime/PakFile/Private/IoDispatcherFileBackend.cpp
+++ b/Engine/Source/Runtime/PakFile/Private/IoDispatcherFileBackend.cpp
@@ -1542,6 +1542,11 @@ void FFileIoStore::CancelIoRequest(FIoRequestImpl* Request)
 			ResolvedRequest->bCancelled = true;
 			CompleteDispatcherRequest(ResolvedRequest);
 		}
+		else
+		{
+			// Wake-up the I/O thread to process cancelled read requests
+			PlatformImpl->ServiceNotify();
+		}
 	}
 }
 
