PakFile: Fix FindFilesAtPathInIndex to handle unnormalized directory queries that do not end in /.
#rnx

diff --git a/Engine/Source/Runtime/PakFile/Private/PakFile.inl b/Engine/Source/Runtime/PakFile/Private/PakFile.inl
--- a/Engine/Source/Runtime/PakFile/Private/PakFile.inl
+++ b/Engine/Source/Runtime/PakFile/Private/PakFile.inl
@@ -12,19 +12,31 @@ void FPakFile::FindFilesAtPathInIndex(const FDirectoryIndex& TargetIndex, Contai
 	const FString& Directory, const ShouldVisitFunc& ShouldVisit, bool bIncludeFiles, bool bIncludeDirectories,
 	bool bRecursive) const
 {
-	// Early out if MountPoint is not matching directory
-	if (!Directory.StartsWith(MountPoint))
+	FStringView RelativeSearch;
+	if (Directory.StartsWith(MountPoint))
 	{
-		return;
+		RelativeSearch = FStringView(Directory).RightChop(MountPoint.Len());
+	}
+	else
+	{
+		// Directory is unnormalized and might not end with /; MountPoint is guaranteed to end with /.
+		// Act as if we were called with a normalized directory if adding slash makes it match MountPoint.
+		if (FStringView(Directory).StartsWith(FStringView(MountPoint).LeftChop(1)))
+		{
+			RelativeSearch = FStringView();
+		}
+		else
+		{
+			// Early out; directory does not start with MountPoint and so will not match any of files in this pakfile.
+			return;
+		}
 	}
-
-	FStringView RelativeSearch(FStringView(Directory).RightChop(MountPoint.Len()));
 
 	TArray<FString> DirectoriesInPak; // List of all unique directories at path
 	for (TMap<FString, FPakDirectory>::TConstIterator It(TargetIndex); It; ++It)
 	{
 		// Check if the file is under the specified path.
-		if (FStringView(It.Key()).StartsWith(RelativeSearch))
+		if (RelativeSearch.IsEmpty() || FStringView(It.Key()).StartsWith(RelativeSearch))
 		{
 			FString PakPath = PakPathCombine(MountPoint, It.Key());
 			if (bRecursive == true)
