[AutoRTFM] Use transactionally-safe critical sections to handle FilePaks.

diff --git a/Engine/Source/Runtime/PakFile/Private/IPlatformFilePak.cpp b/Engine/Source/Runtime/PakFile/Private/IPlatformFilePak.cpp
--- a/Engine/Source/Runtime/PakFile/Private/IPlatformFilePak.cpp
+++ b/Engine/Source/Runtime/PakFile/Private/IPlatformFilePak.cpp
@@ -14,6 +14,7 @@
 #include "Misc/App.h"
 #include "Misc/ConfigCacheIni.h"
 #include "Misc/SecureHash.h"
+#include "Misc/TransactionallySafeScopeLock.h"
 #include "HAL/FileManagerGeneric.h"
 #include "SignedArchiveReader.h"
 #include "Misc/AES.h"
@@ -4495,7 +4496,7 @@ FPakPlatformFile::~FPakPlatformFile()
 	FPakPrecacher::Shutdown();
 #endif
 	{
-		FScopeLock ScopedLock(&PakListCritical);
+		FTransactionallySafeScopeLock ScopedLock(&PakListCritical);
 		for (int32 PakFileIndex = 0; PakFileIndex < PakFiles.Num(); PakFileIndex++)
 		{
 			PakFiles[PakFileIndex].PakFile.SafeRelease();
@@ -4979,7 +4980,7 @@ void FPakPlatformFile::GetFilenamesFromIostoreContainer(const FString& InContain
 
 	const TMap<FGuid, FAES::FAESKey> Keys = UE::FEncryptionKeyManager::Get().GetAllKeys();
 
-	FScopeLock ScopedLock(&PakPlatformFile->PakListCritical);
+	FTransactionallySafeScopeLock ScopedLock(&PakPlatformFile->PakListCritical);
 	for (const FPakListEntry& PakListEntry : PakPlatformFile->PakFiles)
 	{
 		if (FPaths::GetBaseFilename(PakListEntry.PakFile->PakFilename) == InContainerName)
@@ -5005,7 +5006,7 @@ void FPakPlatformFile::ForeachPackageInIostoreWhile(TFunctionRef<bool(FName)> Pr
 
 	const TMap<FGuid, FAES::FAESKey> Keys = UE::FEncryptionKeyManager::Get().GetAllKeys();
 
-	FScopeLock ScopedLock(&PakPlatformFile->PakListCritical);
+	FTransactionallySafeScopeLock ScopedLock(&PakPlatformFile->PakListCritical);
 	for (const FPakListEntry& PakListEntry : PakPlatformFile->PakFiles)
 	{
 		TUniquePtr<FIoStoreReader> IoStoreReader(new FIoStoreReader());
@@ -5215,7 +5216,7 @@ void FPakPlatformFile::GetFilenamesFromIostoreByBlockIndex(const FString& InCont
 
 	const TMap<FGuid, FAES::FAESKey> Keys = UE::FEncryptionKeyManager::Get().GetAllKeys();
 
-	FScopeLock ScopedLock(&PakPlatformFile->PakListCritical);
+	FTransactionallySafeScopeLock ScopedLock(&PakPlatformFile->PakListCritical);
 	for (const FPakListEntry& PakListEntry : PakPlatformFile->PakFiles)
 	{
 		if (FPaths::GetBaseFilename(PakListEntry.PakFile->PakFilename) == InContainerName)
@@ -5838,7 +5839,7 @@ bool FPakPlatformFile::Mount(const TCHAR* InPakFilename, uint32 PakOrder, const
 				}
 				{
 					// Add new pak file
-					FScopeLock ScopedLock(&PakListCritical);
+					FTransactionallySafeScopeLock ScopedLock(&PakListCritical);
 					FPakListEntry Entry;
 					Entry.ReadOrder = PakOrder;
 					Entry.PakFile = Pak;
@@ -6009,7 +6010,7 @@ bool FPakPlatformFile::Unmount(const TCHAR* InPakFilename)
 	TRefCountPtr<FPakFile> UnmountedPak;
 	bool bRemovedContainerFile = false;
 	{
-		FScopeLock ScopedLock(&PakListCritical);
+		FTransactionallySafeScopeLock ScopedLock(&PakListCritical);
 		for (int32 PakIndex = 0; PakIndex < PakFiles.Num(); PakIndex++)
 		{
 			FPakListEntry& PakListEntry = PakFiles[PakIndex];
@@ -6361,7 +6362,7 @@ void FPakPlatformFile::SetMountStartupPaksWildCard(const FString& WildCard)
 
 EChunkLocation::Type FPakPlatformFile::GetPakChunkLocation(int32 InPakchunkIndex) const
 {
-	FScopeLock ScopedLock(&PakListCritical);
+	FTransactionallySafeScopeLock ScopedLock(&PakListCritical);
 
 	for (const FPakListEntry& PakEntry : PakFiles)
 	{
@@ -6384,7 +6385,7 @@ EChunkLocation::Type FPakPlatformFile::GetPakChunkLocation(int32 InPakchunkIndex
 
 bool FPakPlatformFile::AnyChunksAvailable() const
 {
-	FScopeLock ScopedLock(&PakListCritical);
+	FTransactionallySafeScopeLock ScopedLock(&PakListCritical);
 
 	for (const FPakListEntry& PakEntry : PakFiles)
 	{
diff --git a/Engine/Source/Runtime/PakFile/Public/IPlatformFilePak.h b/Engine/Source/Runtime/PakFile/Public/IPlatformFilePak.h
--- a/Engine/Source/Runtime/PakFile/Public/IPlatformFilePak.h
+++ b/Engine/Source/Runtime/PakFile/Public/IPlatformFilePak.h
@@ -35,6 +35,8 @@
 #include "Misc/ScopeLock.h"
 #include "Misc/ScopeRWLock.h"
 #include "Misc/SecureHash.h"
+#include "Misc/TransactionallySafeCriticalSection.h"
+#include "Misc/TransactionallySafeScopeLock.h"
 #include "RSA.h"
 #include "Serialization/Archive.h"
 #include "Serialization/MemoryImage.h"
@@ -1949,7 +1951,7 @@ class FPakPlatformFile : public IPlatformFile
 	/** True if this we're using signed content. */
 	bool bSigned;
 	/** Synchronization object for accessing the list of currently mounted pak files. */
-	mutable FCriticalSection PakListCritical;
+	mutable FTransactionallySafeCriticalSection PakListCritical;
 	/** Cache of extensions that we automatically reject if not found in pak file */
 	TSet<FName> ExcludedNonPakExtensions;
 	/** The extension used for ini files, used for excluding ini files */
@@ -1972,7 +1974,7 @@ class FPakPlatformFile : public IPlatformFile
 	 */
 	FORCEINLINE void GetMountedPaks(TArray<FPakListEntry>& Paks)
 	{
-		FScopeLock ScopedLock(&PakListCritical);
+		FTransactionallySafeScopeLock ScopedLock(&PakListCritical);
 		Paks.Append(PakFiles);
 	}
 
@@ -2102,7 +2104,7 @@ public:
 	*/
 	FORCEINLINE void GetMountedPakFilenames(TArray<FString>& PakFilenames)
 	{
-		FScopeLock ScopedLock(&PakListCritical);
+		FTransactionallySafeScopeLock ScopedLock(&PakListCritical);
 		PakFilenames.Empty(PakFiles.Num());
 		for (FPakListEntry& Entry : PakFiles)
 		{
@@ -2115,7 +2117,7 @@ public:
 	*/
 	FORCEINLINE void GetMountedPakFilenames(TSet<FString>& PakFilenames)
 	{
-		FScopeLock ScopedLock(&PakListCritical);
+		FTransactionallySafeScopeLock ScopedLock(&PakListCritical);
 		PakFilenames.Empty(PakFiles.Num());
 		for (FPakListEntry& Entry : PakFiles)
 		{
