[AutoRTFM] Migrate to UE::TScopeLock instead of FTransactionallySafeScopeLock.

These are functionally identical.

diff --git a/Engine/Source/Runtime/PakFile/Private/IPlatformFilePak.cpp b/Engine/Source/Runtime/PakFile/Private/IPlatformFilePak.cpp
--- a/Engine/Source/Runtime/PakFile/Private/IPlatformFilePak.cpp
+++ b/Engine/Source/Runtime/PakFile/Private/IPlatformFilePak.cpp
@@ -14,7 +14,7 @@
 #include "Misc/App.h"
 #include "Misc/ConfigCacheIni.h"
 #include "Misc/SecureHash.h"
-#include "Misc/TransactionallySafeScopeLock.h"
+#include "Misc/ScopeLock.h"
 #include "HAL/FileManagerGeneric.h"
 #include "SignedArchiveReader.h"
 #include "Misc/AES.h"
@@ -4496,7 +4496,7 @@ FPakPlatformFile::~FPakPlatformFile()
 	FPakPrecacher::Shutdown();
 #endif
 	{
-		FTransactionallySafeScopeLock ScopedLock(&PakListCritical);
+		UE::TScopeLock ScopedLock(PakListCritical);
 		for (int32 PakFileIndex = 0; PakFileIndex < PakFiles.Num(); PakFileIndex++)
 		{
 			PakFiles[PakFileIndex].PakFile.SafeRelease();
@@ -4981,7 +4981,7 @@ void FPakPlatformFile::GetFilenamesFromIostoreContainer(const FString& InContain
 
 	const TMap<FGuid, FAES::FAESKey> Keys = UE::FEncryptionKeyManager::Get().GetAllKeys();
 
-	FTransactionallySafeScopeLock ScopedLock(&PakPlatformFile->PakListCritical);
+	UE::TScopeLock ScopedLock(PakPlatformFile->PakListCritical);
 	for (const FPakListEntry& PakListEntry : PakPlatformFile->PakFiles)
 	{
 		if (FPaths::GetBaseFilename(PakListEntry.PakFile->PakFilename) == InContainerName)
@@ -5007,7 +5007,7 @@ void FPakPlatformFile::ForeachPackageInIostoreWhile(TFunctionRef<bool(FName)> Pr
 
 	const TMap<FGuid, FAES::FAESKey> Keys = UE::FEncryptionKeyManager::Get().GetAllKeys();
 
-	FTransactionallySafeScopeLock ScopedLock(&PakPlatformFile->PakListCritical);
+	UE::TScopeLock ScopedLock(PakPlatformFile->PakListCritical);
 	for (const FPakListEntry& PakListEntry : PakPlatformFile->PakFiles)
 	{
 		TUniquePtr<FIoStoreReader> IoStoreReader(new FIoStoreReader());
@@ -5217,7 +5217,7 @@ void FPakPlatformFile::GetFilenamesFromIostoreByBlockIndex(const FString& InCont
 
 	const TMap<FGuid, FAES::FAESKey> Keys = UE::FEncryptionKeyManager::Get().GetAllKeys();
 
-	FTransactionallySafeScopeLock ScopedLock(&PakPlatformFile->PakListCritical);
+	UE::TScopeLock ScopedLock(PakPlatformFile->PakListCritical);
 	for (const FPakListEntry& PakListEntry : PakPlatformFile->PakFiles)
 	{
 		if (FPaths::GetBaseFilename(PakListEntry.PakFile->PakFilename) == InContainerName)
@@ -5840,7 +5840,7 @@ bool FPakPlatformFile::Mount(const TCHAR* InPakFilename, uint32 PakOrder, const
 				}
 				{
 					// Add new pak file
-					FTransactionallySafeScopeLock ScopedLock(&PakListCritical);
+					UE::TScopeLock ScopedLock(PakListCritical);
 					FPakListEntry Entry;
 					Entry.ReadOrder = PakOrder;
 					Entry.PakFile = Pak;
@@ -6011,7 +6011,7 @@ bool FPakPlatformFile::Unmount(const TCHAR* InPakFilename)
 	TRefCountPtr<FPakFile> UnmountedPak;
 	bool bRemovedContainerFile = false;
 	{
-		FTransactionallySafeScopeLock ScopedLock(&PakListCritical);
+		UE::TScopeLock ScopedLock(PakListCritical);
 		for (int32 PakIndex = 0; PakIndex < PakFiles.Num(); PakIndex++)
 		{
 			FPakListEntry& PakListEntry = PakFiles[PakIndex];
@@ -6363,7 +6363,7 @@ void FPakPlatformFile::SetMountStartupPaksWildCard(const FString& WildCard)
 
 EChunkLocation::Type FPakPlatformFile::GetPakChunkLocation(int32 InPakchunkIndex) const
 {
-	FTransactionallySafeScopeLock ScopedLock(&PakListCritical);
+	UE::TScopeLock ScopedLock(PakListCritical);
 
 	for (const FPakListEntry& PakEntry : PakFiles)
 	{
@@ -6386,7 +6386,7 @@ EChunkLocation::Type FPakPlatformFile::GetPakChunkLocation(int32 InPakchunkIndex
 
 bool FPakPlatformFile::AnyChunksAvailable() const
 {
-	FTransactionallySafeScopeLock ScopedLock(&PakListCritical);
+	UE::TScopeLock ScopedLock(PakListCritical);
 
 	for (const FPakListEntry& PakEntry : PakFiles)
 	{
diff --git a/Engine/Source/Runtime/PakFile/Public/IPlatformFilePak.h b/Engine/Source/Runtime/PakFile/Public/IPlatformFilePak.h
--- a/Engine/Source/Runtime/PakFile/Public/IPlatformFilePak.h
+++ b/Engine/Source/Runtime/PakFile/Public/IPlatformFilePak.h
@@ -36,7 +36,6 @@
 #include "Misc/ScopeRWLock.h"
 #include "Misc/SecureHash.h"
 #include "Misc/TransactionallySafeCriticalSection.h"
-#include "Misc/TransactionallySafeScopeLock.h"
 #include "RSA.h"
 #include "Serialization/Archive.h"
 #include "Serialization/MemoryImage.h"
@@ -1992,7 +1991,7 @@ class FPakPlatformFile : public IPlatformFile
 	 */
 	FORCEINLINE void GetMountedPaks(TArray<FPakListEntry>& Paks)
 	{
-		FTransactionallySafeScopeLock ScopedLock(&PakListCritical);
+		UE::TScopeLock ScopedLock(PakListCritical);
 		Paks.Append(PakFiles);
 	}
 
@@ -2122,7 +2121,7 @@ public:
 	*/
 	FORCEINLINE void GetMountedPakFilenames(TArray<FString>& PakFilenames)
 	{
-		FTransactionallySafeScopeLock ScopedLock(&PakListCritical);
+		UE::TScopeLock ScopedLock(PakListCritical);
 		PakFilenames.Empty(PakFiles.Num());
 		for (FPakListEntry& Entry : PakFiles)
 		{
@@ -2135,7 +2134,7 @@ public:
 	*/
 	FORCEINLINE void GetMountedPakFilenames(TSet<FString>& PakFilenames)
 	{
-		FTransactionallySafeScopeLock ScopedLock(&PakListCritical);
+		UE::TScopeLock ScopedLock(PakListCritical);
 		PakFilenames.Empty(PakFiles.Num());
 		for (FPakListEntry& Entry : PakFiles)
 		{
