Temp work around for TRefCountPtr needing AutoRTFM support

diff --git a/Engine/Source/Runtime/PakFile/Private/IPlatformFilePak.cpp b/Engine/Source/Runtime/PakFile/Private/IPlatformFilePak.cpp
--- a/Engine/Source/Runtime/PakFile/Private/IPlatformFilePak.cpp
+++ b/Engine/Source/Runtime/PakFile/Private/IPlatformFilePak.cpp
@@ -5316,7 +5316,7 @@ bool FPakPlatformFile::FindFileInPakFiles(TArray<FPakListEntry>& Paks, const TCH
 	return false;
 }
 
-bool FPakPlatformFile::FindFileInPakFiles(const TCHAR* Filename, TRefCountPtr<FPakFile>* OutPakFile,
+UE_AUTORTFM_ALWAYS_OPEN bool FPakPlatformFile::FindFileInPakFiles(const TCHAR* Filename, TRefCountPtr<FPakFile>* OutPakFile,
 	FPakEntry* OutEntry)
 {
 	TArray<FPakListEntry> Paks;
