Added an overload to IPlatformFile::OpenMapped that return both a value and an error code. Old approach is deprecated now

diff --git a/Engine/Source/Developer/TraceServices/Private/Model/SymslibResolver.cpp b/Engine/Source/Developer/TraceServices/Private/Model/SymslibResolver.cpp
--- a/Engine/Source/Developer/TraceServices/Private/Model/SymslibResolver.cpp
+++ b/Engine/Source/Developer/TraceServices/Private/Model/SymslibResolver.cpp
@@ -58,8 +58,13 @@ namespace
 
 		bool Load(const TCHAR* FileName)
 		{
-			Handle.Reset(FPlatformFileManager::Get().GetPlatformFile().OpenMapped(FileName));
-			Region.Reset(Handle.IsValid() ? Handle->MapRegion(0, Handle->GetFileSize()) : nullptr);
+			FOpenMappedResult Result = FPlatformFileManager::Get().GetPlatformFile().OpenMappedEx(FileName);
+			if (Result.HasError())
+			{
+				return false;
+			}
+			Handle = Result.StealValue();
+			Region.Reset(Handle->MapRegion(0, Handle->GetFileSize()));
 			return Region.IsValid();
 		}
 
diff --git a/Engine/Source/Programs/SymsLibDump/Private/SymsLibDump.cpp b/Engine/Source/Programs/SymsLibDump/Private/SymsLibDump.cpp
--- a/Engine/Source/Programs/SymsLibDump/Private/SymsLibDump.cpp
+++ b/Engine/Source/Programs/SymsLibDump/Private/SymsLibDump.cpp
@@ -41,14 +41,16 @@ namespace
 
 		bool Load(const TCHAR* FileName)
 		{
-			Handle.Reset(FPlatformFileManager::Get().GetPlatformFile().OpenMapped(FileName));
+			FOpenMappedResult Result = FPlatformFileManager::Get().GetPlatformFile().OpenMappedEx(FileName);
 
-			if (Handle.IsValid())
+			if (!Result.HasError())
 			{
+				Handle = Result.StealValue();
 				Region.Reset(Handle->MapRegion(0, Handle->GetFileSize()));
 			}
 			else
 			{
+				Handle = nullptr;
 				Region.Reset();
 			}
 
diff --git a/Engine/Source/Runtime/Apple/MetalRHI/Private/MetalDynamicRHI_Shaders.cpp b/Engine/Source/Runtime/Apple/MetalRHI/Private/MetalDynamicRHI_Shaders.cpp
--- a/Engine/Source/Runtime/Apple/MetalRHI/Private/MetalDynamicRHI_Shaders.cpp
+++ b/Engine/Source/Runtime/Apple/MetalRHI/Private/MetalDynamicRHI_Shaders.cpp
@@ -158,7 +158,8 @@ FRHIShaderLibraryRef FMetalDynamicRHI::RHICreateShaderLibrary(EShaderPlatform Pl
     TUniquePtr<FMetalShaderLibrary::FShaderLibDataOwner> MemOwner = MakeUnique<FMetalShaderLibrary::FShaderLibDataOwner>();
 	if (FPlatformProperties::SupportsMemoryMappedFiles())
 	{
-		MemOwner->MappedCacheFile = TUniquePtr<IMappedFileHandle>(FPlatformFileManager::Get().GetPlatformFile().OpenMapped(*BinaryShaderFile));
+		FOpenMappedResult Result = FPlatformFileManager::Get().GetPlatformFile().OpenMappedEx(*BinaryShaderFile);
+		MemOwner->MappedCacheFile = Result.HasError() ? nullptr : Result.StealValue();
 		if (MemOwner->MappedCacheFile.IsValid())
 		{
 			MemOwner->MappedRegion = TUniquePtr<IMappedFileRegion>(MemOwner->MappedCacheFile->MapRegion(0, MemOwner->MappedCacheFile->GetFileSize()));
diff --git a/Engine/Source/Runtime/AssetRegistry/Private/AssetDataGatherer.cpp b/Engine/Source/Runtime/AssetRegistry/Private/AssetDataGatherer.cpp
--- a/Engine/Source/Runtime/AssetRegistry/Private/AssetDataGatherer.cpp
+++ b/Engine/Source/Runtime/AssetRegistry/Private/AssetDataGatherer.cpp
@@ -3500,9 +3500,14 @@ class FMemoryMappedFile
 {
 public:
 	FMemoryMappedFile(const TCHAR* Path)
-		: Handle(FPlatformFileManager::Get().GetPlatformFile().OpenMapped(Path))
-		, Region(Handle ? Handle->MapRegion() : nullptr)
-	{}
+	{
+		FOpenMappedResult Result = FPlatformFileManager::Get().GetPlatformFile().OpenMappedEx(Path);
+		if (!Result.HasError())
+		{
+			Handle = Result.StealValue();
+			Region.Reset(Handle->MapRegion());
+		}
+	}
 
 	void Preload(int64 Size = MAX_int64) const
 	{
diff --git a/Engine/Source/Runtime/Core/Private/Android/AndroidPlatformFile.cpp b/Engine/Source/Runtime/Core/Private/Android/AndroidPlatformFile.cpp
--- a/Engine/Source/Runtime/Core/Private/Android/AndroidPlatformFile.cpp
+++ b/Engine/Source/Runtime/Core/Private/Android/AndroidPlatformFile.cpp
@@ -1447,7 +1447,7 @@ public:
 		return result;
 	}
 
-	virtual IMappedFileHandle* OpenMapped(const TCHAR* Filename, EOpenReadFlags OpenOptions = EOpenReadFlags::None) override
+	virtual FOpenMappedResult OpenMappedEx(const TCHAR* Filename, EOpenReadFlags OpenOptions = EOpenReadFlags::None) override
 	{
 		FString LocalPath;
 		FString AssetPath;
@@ -1460,11 +1460,11 @@ public:
 		if (Handle == -1)
 		{
 			const int ErrNo = errno;
+			FString ErrorStr = FString::Printf(TEXT("FAndroidPlatformFile::OpenMappedEx('%s', Flags=0x%08X) failed: errno=%d (%s)"), *NormalizedFilename, Flags, ErrNo, UTF8_TO_TCHAR(strerror(ErrNo));
 #if LOG_ANDROID_FILE
-			FPlatformMisc::LowLevelOutputDebugStringf(TEXT("FAndroidPlatformFile::OpenMapped('%s', Flags=0x%08X) failed: errno=%d (%s)"), *NormalizedFilename, Flags, ErrNo, UTF8_TO_TCHAR(strerror(ErrNo)));
+			FPlatformMisc::LowLevelOutputDebugString(*ErrorStr);
 #endif
-
-			return nullptr;
+			return MakeError(MoveTemp(ErrorStr));
 		}
 
 		struct stat FileInfo;
@@ -1473,15 +1473,16 @@ public:
 		if (StatResult == -1)
 		{
 			const int ErrNo = errno;
+			FString ErrorStr = FString::Printf(TEXT("FAndroidPlatformFile::OpenMappedEx fstat failed: ('%s', Flags=0x%08X) failed: errno=%d (%s)"), *NormalizedFilename, Flags, ErrNo, UTF8_TO_TCHAR(strerror(ErrNo));
 
 #if LOG_ANDROID_FILE
-			FPlatformMisc::LowLevelOutputDebugStringf(TEXT("FAndroidPlatformFile::OpenMapped fstat failed: ('%s', Flags=0x%08X) failed: errno=%d (%s)"), *NormalizedFilename, Flags, ErrNo, UTF8_TO_TCHAR(strerror(ErrNo)));
+			FPlatformMisc::LowLevelOutputDebugString(*ErrorStr);
 #endif
 
-			return nullptr;
+			return MakeError(MoveTemp(ErrorStr));
 		}
 
-		return new FAndroidMappedFileHandle(Handle, FileInfo.st_size, NormalizedFilename);
+		return MakeValue(new FAndroidMappedFileHandle(Handle, FileInfo.st_size, NormalizedFilename));
 	}
 
 	virtual int64 FileSize(const TCHAR* Filename) override
diff --git a/Engine/Source/Runtime/Core/Private/GenericPlatform/GenericPlatformFile.cpp b/Engine/Source/Runtime/Core/Private/GenericPlatform/GenericPlatformFile.cpp
--- a/Engine/Source/Runtime/Core/Private/GenericPlatform/GenericPlatformFile.cpp
+++ b/Engine/Source/Runtime/Core/Private/GenericPlatform/GenericPlatformFile.cpp
@@ -537,6 +537,22 @@ IAsyncReadFileHandle* IPlatformFile::OpenAsyncRead(const TCHAR* Filename, bool b
 	return new FGenericAsyncReadFileHandle(this, Filename, bAllowWrite);
 }
 
+IMappedFileHandle* IPlatformFile::OpenMapped(const TCHAR* Filename)
+{
+	FOpenMappedResult Result = OpenMappedEx(Filename);
+	if (Result.HasError())
+	{
+		return nullptr;
+	}
+
+	return Result.StealValue().Get();
+}
+
+FOpenMappedResult IPlatformFile::OpenMappedEx(const TCHAR* Filename, EOpenReadFlags OpenOptions)
+{
+	return MakeError(TEXTVIEW("OpenMappedEx is not implemented on this platform"));
+}
+
 DEFINE_STAT(STAT_AsyncFileMemory);
 DEFINE_STAT(STAT_AsyncFileHandles);
 DEFINE_STAT(STAT_AsyncFileRequests);
diff --git a/Engine/Source/Runtime/Core/Private/HAL/IPlatformFileLogWrapper.cpp b/Engine/Source/Runtime/Core/Private/HAL/IPlatformFileLogWrapper.cpp
--- a/Engine/Source/Runtime/Core/Private/HAL/IPlatformFileLogWrapper.cpp
+++ b/Engine/Source/Runtime/Core/Private/HAL/IPlatformFileLogWrapper.cpp
@@ -3,6 +3,7 @@
 #include "HAL/IPlatformFileLogWrapper.h"
 #include "Misc/CoreMisc.h"
 #include "Templates/UniquePtr.h"
+#include "Async/MappedFileHandle.h"
 
 bool bSuppressFileLog = false;
 DEFINE_LOG_CATEGORY(LogPlatformFile);
@@ -86,3 +87,8 @@ FLoggedFileHandle::~FLoggedFileHandle()
 #endif
 	FILE_LOG(LogPlatformFile, Log, TEXT("Close %s"), *Filename);
 }
+
+FOpenMappedResult FLoggedPlatformFile::OpenMappedEx(const TCHAR* Filename, EOpenReadFlags OpenOptions)
+{
+	return LowerLevel->OpenMappedEx(Filename, OpenOptions);
+}
diff --git a/Engine/Source/Runtime/Core/Private/HAL/IPlatformFileOpenLogWrapper.cpp b/Engine/Source/Runtime/Core/Private/HAL/IPlatformFileOpenLogWrapper.cpp
--- a/Engine/Source/Runtime/Core/Private/HAL/IPlatformFileOpenLogWrapper.cpp
+++ b/Engine/Source/Runtime/Core/Private/HAL/IPlatformFileOpenLogWrapper.cpp
@@ -4,6 +4,8 @@
 
 #if !UE_BUILD_SHIPPING
 
+#include "Async/MappedFileHandle.h"
+
 DEFINE_LOG_CATEGORY_STATIC(LogFileOpenOrder, Log, All);
 
 static FPlatformFileOpenLog* GPlatformFileOpenLog = nullptr;
@@ -104,6 +106,11 @@ bool FPlatformFileOpenLog::Initialize(IPlatformFile* Inner, const TCHAR* Command
 	return true;
 }
 
+FOpenMappedResult FPlatformFileOpenLog::OpenMappedEx(const TCHAR* Filename, EOpenReadFlags OpenOptions)
+{
+	return LowerLevel->OpenMappedEx(Filename, OpenOptions);
+}
+
 void FPlatformFileOpenLog::AddLabelInternal(const TCHAR* Fmt, ...)
 {
 	if (GPlatformFileOpenLog)
diff --git a/Engine/Source/Runtime/Core/Private/IOS/IOSPlatformFile.cpp b/Engine/Source/Runtime/Core/Private/IOS/IOSPlatformFile.cpp
--- a/Engine/Source/Runtime/Core/Private/IOS/IOSPlatformFile.cpp
+++ b/Engine/Source/Runtime/Core/Private/IOS/IOSPlatformFile.cpp
@@ -910,7 +910,7 @@ IFileHandle* FIOSPlatformFile::OpenWrite(const TCHAR* Filename, bool bAppend, bo
 	return NULL;
 }
 
-IMappedFileHandle* FIOSPlatformFile::OpenMapped(const TCHAR* Filename, EOpenReadFlags OpenOptions)
+FOpenMappedResult FIOSPlatformFile::OpenMappedEx(const TCHAR* Filename, EOpenReadFlags OpenOptions)
 {
 	FString NormalizedFilename = NormalizeFilename(Filename);
 	
@@ -948,10 +948,11 @@ IMappedFileHandle* FIOSPlatformFile::OpenMapped(const TCHAR* Filename, EOpenRead
 			}
 			uint64 FileSize = FileInfo.st_size;
 
-			return new FIOSMappedFileHandle(Handle, FileSize, FinalPath);
+			return MakeValue(new FIOSMappedFileHandle(Handle, FileSize, FinalPath));
 		}
 	}
-	return NULL;
+	
+	return MakeError(FString::Printf(TEXT("FIOSPlatformFile::OpenMappedEx failed to open file '%s' in '%s' mode"), *NormalizedFilename, EnumHasAnyFlags(OpenOptions, EOpenReadFlags::AllowWrite) ? TEXTVIEW("write") : TEXTVIEW("read")));
 }
 
 bool FIOSPlatformFile::DirectoryExists(const TCHAR* Directory)
diff --git a/Engine/Source/Runtime/Core/Private/Unix/UnixPlatformFile.cpp b/Engine/Source/Runtime/Core/Private/Unix/UnixPlatformFile.cpp
--- a/Engine/Source/Runtime/Core/Private/Unix/UnixPlatformFile.cpp
+++ b/Engine/Source/Runtime/Core/Private/Unix/UnixPlatformFile.cpp
@@ -1226,7 +1226,7 @@ IFileHandle* FUnixPlatformFile::OpenWrite(const TCHAR* Filename, bool bAppend, b
 	return nullptr;
 }
 
-IMappedFileHandle* FUnixPlatformFile::OpenMapped(const TCHAR* Filename, EOpenReadFlags OpenOptions)
+FOpenMappedResult FUnixPlatformFile::OpenMappedEx(const TCHAR* Filename, EOpenReadFlags OpenOptions)
 {
 	const FString NormalizedFilename = NormalizeFilename(Filename, false);
 
@@ -1235,9 +1235,10 @@ IMappedFileHandle* FUnixPlatformFile::OpenMapped(const TCHAR* Filename, EOpenRea
 	if (Handle == -1)
 	{
 		const int ErrNo = errno;
-		UE_LOG_UNIX_FILE(Warning, TEXT("open('%s', Flags=0x%08X) failed: errno=%d (%s)"), *NormalizedFilename, Flags, ErrNo, UTF8_TO_TCHAR(strerror(ErrNo)));
+		FString ErrorStr = FString::Printf(TEXT("open('%s', Flags=0x%08X) failed: errno=%d (%s)"), *NormalizedFilename, Flags, ErrNo, UTF8_TO_TCHAR(strerror(ErrNo)));
+		UE_LOG_UNIX_FILE(Warning, *ErrorStr);
 
-		return nullptr;
+		return MakeError(ErrorStr);
 	}
 	
 	struct stat FileInfo;
@@ -1246,12 +1247,13 @@ IMappedFileHandle* FUnixPlatformFile::OpenMapped(const TCHAR* Filename, EOpenRea
 	if (StatResult == -1)
 	{
 		const int ErrNo = errno;
-		UE_LOG_UNIX_FILE(Warning, TEXT("stat('%s', Flags=0x%08X) failed: errno=%d (%s)"), *NormalizedFilename, Flags, ErrNo, UTF8_TO_TCHAR(strerror(ErrNo)));
+		FString ErrorStr = FString::Printf(TEXT("stat('%s', Flags=0x%08X) failed: errno=%d (%s)"), *NormalizedFilename, Flags, ErrNo, UTF8_TO_TCHAR(strerror(ErrNo)));
+		UE_LOG_UNIX_FILE(Warning, ErrorStr);
 
-		return nullptr;
+		return MakeError(ErrorStr);
 	}
 	
-	return new FUnixMappedFileHandle(Handle, FileInfo.st_size, NormalizedFilename); 
+	return MakeValue(new FUnixMappedFileHandle(Handle, FileInfo.st_size, NormalizedFilename));
 }
 
 bool FUnixPlatformFile::DirectoryExists(const TCHAR* Directory)
diff --git a/Engine/Source/Runtime/Core/Private/Windows/WindowsPlatformFile.cpp b/Engine/Source/Runtime/Core/Private/Windows/WindowsPlatformFile.cpp
--- a/Engine/Source/Runtime/Core/Private/Windows/WindowsPlatformFile.cpp
+++ b/Engine/Source/Runtime/Core/Private/Windows/WindowsPlatformFile.cpp
@@ -1508,12 +1508,12 @@ IFileHandle* FWindowsPlatformFile::OpenWrite(const TCHAR* Filename, bool bAppend
 	}
 }
 
-IMappedFileHandle* FWindowsPlatformFile::OpenMapped(const TCHAR* Filename, EOpenReadFlags OpenOptions)
+FOpenMappedResult FWindowsPlatformFile::OpenMappedEx(const TCHAR* Filename, EOpenReadFlags OpenOptions)
 {
 	int64 Size = FileSize(Filename);
 	if (Size < 1)
 	{
-		return nullptr;
+		return MakeError(TEXTVIEW("File size is 0"));
 	}
 	uint32 Access = GENERIC_READ;
 	if (EnumHasAnyFlags(OpenOptions, EOpenReadFlags::AllowWrite))
@@ -1522,8 +1522,10 @@ IMappedFileHandle* FWindowsPlatformFile::OpenMapped(const TCHAR* Filename, EOpen
 	}
 	uint32 WinFlags = FILE_SHARE_READ;
 	uint32 Create = OPEN_EXISTING;
+	FNormalizedFilename NormalizedFilename = FNormalizedFilename(Filename);
+
 	TRACE_PLATFORMFILE_BEGIN_OPEN(Filename);
-	HANDLE Handle = CreateFileW(*FNormalizedFilename(Filename), Access, WinFlags, NULL, Create, FILE_ATTRIBUTE_NORMAL, NULL);
+	HANDLE Handle = CreateFileW(*NormalizedFilename, Access, WinFlags, NULL, Create, FILE_ATTRIBUTE_NORMAL, NULL);
 	if (Handle != INVALID_HANDLE_VALUE)
 	{
 		TRACE_PLATFORMFILE_END_OPEN(Handle);
@@ -1531,7 +1533,7 @@ IMappedFileHandle* FWindowsPlatformFile::OpenMapped(const TCHAR* Filename, EOpen
 	else
 	{
 		TRACE_PLATFORMFILE_FAIL_OPEN(Filename);
-		return nullptr;
+		return MakeError(FString::Printf(TEXT("Failed to open file '%s'"), *NormalizedFilename));
 	}
 	HANDLE MappingHandle = CreateFileMapping(Handle, NULL, EnumHasAnyFlags(OpenOptions, EOpenReadFlags::AllowWrite) ? PAGE_READWRITE : PAGE_READONLY, 0, 0, NULL);
 	if (MappingHandle == NULL)
@@ -1554,9 +1556,9 @@ IMappedFileHandle* FWindowsPlatformFile::OpenMapped(const TCHAR* Filename, EOpen
 #else
 		(void)CloseResult;
 #endif
-		return nullptr;
+		return MakeError(FString::Printf(TEXT("Failed to create file mapping for a file'%s'"), *NormalizedFilename));
 	}
-	return new FMappedFileHandleWindows(Handle, MappingHandle, Size, Filename);
+	return MakeValue(new FMappedFileHandleWindows(Handle, MappingHandle, Size, Filename));
 }
 
 bool FWindowsPlatformFile::DirectoryExists(const TCHAR* Directory)
diff --git a/Engine/Source/Runtime/Core/Private/Windows/WindowsPlatformFilePrivate.h b/Engine/Source/Runtime/Core/Private/Windows/WindowsPlatformFilePrivate.h
--- a/Engine/Source/Runtime/Core/Private/Windows/WindowsPlatformFilePrivate.h
+++ b/Engine/Source/Runtime/Core/Private/Windows/WindowsPlatformFilePrivate.h
@@ -68,7 +68,7 @@ public:
 	virtual IFileHandle* OpenRead(const TCHAR* Filename, bool bAllowWrite = false) override;
 	virtual IFileHandle* OpenReadNoBuffering(const TCHAR* Filename, bool bAllowWrite = false) override;
 	virtual IFileHandle* OpenWrite(const TCHAR* Filename, bool bAppend = false, bool bAllowRead = false) override;
-	virtual IMappedFileHandle* OpenMapped(const TCHAR* Filename, EOpenReadFlags OpenOptions = EOpenReadFlags::None) override;
+	virtual FOpenMappedResult OpenMappedEx(const TCHAR* Filename, EOpenReadFlags OpenOptions = EOpenReadFlags::None) override;
 	virtual bool DirectoryExists(const TCHAR* Directory) override;
 	virtual bool CreateDirectory(const TCHAR* Directory) override;
 	virtual bool DeleteDirectory(const TCHAR* Directory) override;
diff --git a/Engine/Source/Runtime/Core/Public/GenericPlatform/GenericPlatformFile.h b/Engine/Source/Runtime/Core/Public/GenericPlatform/GenericPlatformFile.h
--- a/Engine/Source/Runtime/Core/Public/GenericPlatform/GenericPlatformFile.h
+++ b/Engine/Source/Runtime/Core/Public/GenericPlatform/GenericPlatformFile.h
@@ -20,8 +20,8 @@
 #include "Templates/ValueOrError.h"
 
 class FArchive;
-class IAsyncReadFileHandle;
 class IMappedFileHandle;
+class IAsyncReadFileHandle;
 namespace ELogVerbosity { enum Type : uint8; }
 
 /**
@@ -274,7 +274,7 @@ struct FFileJournalData
 	bool bIsDirectory : 1;
 };
 
-/** Stores custom error messages from the engine along with an optional system error code that can provider more detailed infomation */
+/** Stores custom error messages from the engine along with an optional system error code that can provider more detailed information */
 class FFileSystemError
 {
 public:
@@ -333,6 +333,9 @@ private:
 using FFileOpenResult = TValueOrError<TUniquePtr<IFileHandle>, FFileSystemError>;
 using FFileOpenAsyncResult = TValueOrError<TUniquePtr<IAsyncReadFileHandle>, FFileSystemError>;
 
+/** Data structure returned by calls to IPlatformFile::OpenMapped */
+using FOpenMappedResult = TValueOrError<TUniquePtr<IMappedFileHandle>, FFileSystemError>;
+
 /**
 * File I/O Interface
 **/
@@ -742,16 +745,22 @@ public:
 	/** Open a file for async reading. This call does hit the disk; it is synchronous open. 
 	*
 	* @param Filename		File to be mapped. This doesn't actually map anything, just opens the file.
-	* @param OpenOptions	Allows specifying map to be writable via AllowWrite or be ReadOnly by deafault.
-	* @return Close the file by delete'ing the handle. A non-null return value does mean the file exists. 
+	* @return Close the file by deleting the handle. A non-null return value does mean the file exists. 
 	* Null can be returned for many reasons even if the file exists. Perhaps this platform does not support mapped files, or this file is compressed in a pak file.
 	* Generally you attempt to open mapped, and if that fails, then use other file operations instead.
 	*/
-	virtual IMappedFileHandle* OpenMapped(const TCHAR* Filename, EOpenReadFlags OpenOptions = EOpenReadFlags::None)
-	{
-		return nullptr;
-	}
+	UE_DEPRECATED(5.6, "Please use OpenMappedEx instead")
+	CORE_API virtual IMappedFileHandle* OpenMapped(const TCHAR* Filename);
 
+	/** Open a file for async reading. This call does hit the disk; it is synchronous open.
+	*
+	* @param Filename		File to be mapped. This doesn't actually map anything, just opens the file.
+	* @param OpenOptions	Allows specifying map to be writable via AllowWrite or be ReadOnly by deafault.
+	* @return Close the file by deleting the handle. A non-null return value does mean the file exists.
+	* Null can be returned for many reasons even if the file exists. Perhaps this platform does not support mapped files, or this file is compressed in a pak file.
+	* Generally you attempt to open mapped, and if that fails, then use other file operations instead.
+	*/
+	CORE_API virtual FOpenMappedResult OpenMappedEx(const TCHAR* Filename, EOpenReadFlags OpenOptions = EOpenReadFlags::None);
 
 	CORE_API virtual void GetTimeStampPair(const TCHAR* PathA, const TCHAR* PathB, FDateTime& OutTimeStampA, FDateTime& OutTimeStampB);
 
diff --git a/Engine/Source/Runtime/Core/Public/HAL/IPlatformFileCachedWrapper.h b/Engine/Source/Runtime/Core/Public/HAL/IPlatformFileCachedWrapper.h
--- a/Engine/Source/Runtime/Core/Public/HAL/IPlatformFileCachedWrapper.h
+++ b/Engine/Source/Runtime/Core/Public/HAL/IPlatformFileCachedWrapper.h
@@ -15,6 +15,7 @@
 #include "GenericPlatform/GenericPlatformFile.h"
 #include "HAL/IPlatformFileLogWrapper.h"
 #include "Templates/UniquePtr.h"
+#include "Async/MappedFileHandle.h"
 
 class IAsyncReadFileHandle;
 
@@ -455,9 +456,9 @@ public:
 	{
 		return LowerLevel->OpenAsyncRead(Filename, bAllowWrite);
 	}
-	virtual IMappedFileHandle* OpenMapped(const TCHAR* Filename, EOpenReadFlags OpenOptions = EOpenReadFlags::None) override
+	virtual FOpenMappedResult OpenMappedEx(const TCHAR* Filename, EOpenReadFlags OpenOptions = EOpenReadFlags::None) override
 	{
-		return LowerLevel->OpenMapped(Filename, OpenOptions);
+		return LowerLevel->OpenMappedEx(Filename, OpenOptions);
 	}
 	virtual void SetAsyncMinimumPriority(EAsyncIOPriorityAndFlags MinPriority) override
 	{
diff --git a/Engine/Source/Runtime/Core/Public/HAL/IPlatformFileLogWrapper.h b/Engine/Source/Runtime/Core/Public/HAL/IPlatformFileLogWrapper.h
--- a/Engine/Source/Runtime/Core/Public/HAL/IPlatformFileLogWrapper.h
+++ b/Engine/Source/Runtime/Core/Public/HAL/IPlatformFileLogWrapper.h
@@ -503,10 +503,8 @@ public:
 		//@todo no wrapped logging for async file handles (yet)
 		return Result;
 	}
-	virtual IMappedFileHandle* OpenMapped(const TCHAR* Filename, EOpenReadFlags OpenOptions = EOpenReadFlags::None) override
-	{
-		return LowerLevel->OpenMapped(Filename, OpenOptions);
-	}
+	virtual FOpenMappedResult OpenMappedEx(const TCHAR* Filename, EOpenReadFlags OpenOptions = EOpenReadFlags::None) override;
+
 	virtual void SetAsyncMinimumPriority(EAsyncIOPriorityAndFlags Priority) override
 	{
 		LowerLevel->SetAsyncMinimumPriority(Priority);
diff --git a/Engine/Source/Runtime/Core/Public/HAL/IPlatformFileOpenLogWrapper.h b/Engine/Source/Runtime/Core/Public/HAL/IPlatformFileOpenLogWrapper.h
--- a/Engine/Source/Runtime/Core/Public/HAL/IPlatformFileOpenLogWrapper.h
+++ b/Engine/Source/Runtime/Core/Public/HAL/IPlatformFileOpenLogWrapper.h
@@ -223,10 +223,7 @@ public:
 		// we must not record the "open" here...what matters is when we start reading the file!
 		return new FLoggingAsyncReadFileHandle(this, Filename, LowerLevel->OpenAsyncRead(Filename, bAllowWrite));
 	}
-	virtual IMappedFileHandle* OpenMapped(const TCHAR* Filename, EOpenReadFlags OpenOptions = EOpenReadFlags::None) override
-	{
-		return LowerLevel->OpenMapped(Filename, OpenOptions);
-	}
+	virtual FOpenMappedResult OpenMappedEx(const TCHAR* Filename, EOpenReadFlags OpenOptions = EOpenReadFlags::None) override;
 
 	void AddToOpenLog(const TCHAR* Filename)
 	{
diff --git a/Engine/Source/Runtime/Core/Public/IOS/IOSPlatformFile.h b/Engine/Source/Runtime/Core/Public/IOS/IOSPlatformFile.h
--- a/Engine/Source/Runtime/Core/Public/IOS/IOSPlatformFile.h
+++ b/Engine/Source/Runtime/Core/Public/IOS/IOSPlatformFile.h
@@ -42,7 +42,7 @@ public:
 
 	virtual IFileHandle* OpenRead(const TCHAR* Filename, bool bAllowWrite = false) override;
 	virtual IFileHandle* OpenWrite(const TCHAR* Filename, bool bAppend = false, bool bAllowRead = false) override;
-	virtual IMappedFileHandle* OpenMapped(const TCHAR* Filename, EOpenReadFlags OpenOptions = EOpenReadFlags::None) override;
+	virtual FOpenMappedResult OpenMappedEx(const TCHAR* Filename, EOpenReadFlags OpenOptions = EOpenReadFlags::None) override;
 
 	virtual bool DirectoryExists(const TCHAR* Directory) override;
 	virtual bool CreateDirectory(const TCHAR* Directory) override;
diff --git a/Engine/Source/Runtime/Core/Public/Unix/UnixPlatformFile.h b/Engine/Source/Runtime/Core/Public/Unix/UnixPlatformFile.h
--- a/Engine/Source/Runtime/Core/Public/Unix/UnixPlatformFile.h
+++ b/Engine/Source/Runtime/Core/Public/Unix/UnixPlatformFile.h
@@ -46,7 +46,7 @@ public:
 
 	CORE_API virtual IFileHandle* OpenRead(const TCHAR* Filename, bool bAllowWrite = false) override;
 	CORE_API virtual IFileHandle* OpenWrite(const TCHAR* Filename, bool bAppend = false, bool bAllowRead = false) override;
-	CORE_API virtual IMappedFileHandle* OpenMapped(const TCHAR* Filename, EOpenReadFlags OpenOptions = EOpenReadFlags::None) override;
+	CORE_API virtual FOpenMappedResult OpenMappedEx(const TCHAR* Filename, EOpenReadFlags OpenOptions = EOpenReadFlags::None) override;
 
 	CORE_API virtual bool DirectoryExists(const TCHAR* Directory) override;
 	CORE_API virtual bool CreateDirectory(const TCHAR* Directory) override;
diff --git a/Engine/Source/Runtime/CoreUObject/Private/UObject/PackageResourceManagerFile.cpp b/Engine/Source/Runtime/CoreUObject/Private/UObject/PackageResourceManagerFile.cpp
--- a/Engine/Source/Runtime/CoreUObject/Private/UObject/PackageResourceManagerFile.cpp
+++ b/Engine/Source/Runtime/CoreUObject/Private/UObject/PackageResourceManagerFile.cpp
@@ -5,6 +5,7 @@
 #include "CoreMinimal.h"
 
 #include "Async/AsyncFileHandleNull.h"
+#include "Async/MappedFileHandle.h"
 #include "HAL/FileManager.h"
 #include "HAL/PlatformFile.h"
 #include "HAL/PlatformFileManager.h"
@@ -288,8 +289,13 @@ IMappedFileHandle* FPackageResourceManagerFile::OpenMappedHandleToPackage(const
 	IteratePossibleFiles(PackagePath, FBulkDataCookedIndex::Default, PackageSegment, OutUpdatedPath,
 		[PlatformFile, &Result](const TCHAR* Filename, EPackageExtension Extension)
 		{
-			Result = PlatformFile->OpenMapped(Filename);
-			return Result != nullptr;
+			FOpenMappedResult Value = PlatformFile->OpenMappedEx(Filename);
+			if (Value.HasError())
+			{
+				return false;
+			}
+			Result = Value.StealValue().Get();
+			return true;
 		});
 	return Result;
 }
diff --git a/Engine/Source/Runtime/OpenGLDrv/Private/OpenGLProgramBinaryFileCache.cpp b/Engine/Source/Runtime/OpenGLDrv/Private/OpenGLProgramBinaryFileCache.cpp
--- a/Engine/Source/Runtime/OpenGLDrv/Private/OpenGLProgramBinaryFileCache.cpp
+++ b/Engine/Source/Runtime/OpenGLDrv/Private/OpenGLProgramBinaryFileCache.cpp
@@ -652,7 +652,8 @@ TRefCountPtr<FOpenGLProgramBinaryMapping> FOpenGLProgramBinaryCache::GetOrAddFil
 		TUniquePtr<IMappedFileHandle> MappedCacheFile;
 		if(UE::OpenGL::CanMemoryMapGLProgramCache())
 		{
-			MappedCacheFile = TUniquePtr<IMappedFileHandle>(FPlatformFileManager::Get().GetPlatformFile().OpenMapped(*ProgramCacheFilename));
+			FOpenMappedResult Result = FPlatformFileManager::Get().GetPlatformFile().OpenMappedEx(*ProgramCacheFilename);
+			MappedCacheFile = Result.HasError() ? nullptr : Result.StealValue();
 		}
 		CurrentMapping = new FOpenGLProgramBinaryMapping(MoveTemp(MappedCacheFile), ProgramCount);
 		MappedCacheFiles.Add(CurrentShaderPipelineProperties.CacheVersionGuid, CurrentMapping);		
diff --git a/Engine/Source/Runtime/PakFile/Private/IPlatformFilePak.cpp b/Engine/Source/Runtime/PakFile/Private/IPlatformFilePak.cpp
--- a/Engine/Source/Runtime/PakFile/Private/IPlatformFilePak.cpp
+++ b/Engine/Source/Runtime/PakFile/Private/IPlatformFilePak.cpp
@@ -3948,14 +3948,13 @@ static void MappedFileTest(const TArray<FString>& Args)
 
 	while (true)
 	{
-		IMappedFileHandle* Handle = FPlatformFileManager::Get().GetPlatformFile().OpenMapped(*TestFile);
+		TUniquePtr<IMappedFileHandle> Handle = FPlatformFileManager::Get().GetPlatformFile().OpenMappedEx(*TestFile).StealValue();
 		IMappedFileRegion *Region = Handle->MapRegion();
 
 		int64 Size = Region->GetMappedSize();
 		const char* Data = (const char *)Region->GetMappedPtr();
 
 		delete Region;
-		delete Handle;
 	}
 
 
@@ -4590,20 +4589,20 @@ bool FPakPlatformFile::IsPakFileInstalled(const FString& InFilename)
 }
 #endif //HAS_PLATFORM_PAK_INSTALL_CHECK
 
-IMappedFileHandle* FPakPlatformFile::OpenMapped(const TCHAR* Filename, EOpenReadFlags OpenOptions)
+FOpenMappedResult FPakPlatformFile::OpenMappedEx(const TCHAR* Filename, EOpenReadFlags OpenOptions)
 {
 	if (!GMMIO_Enable)
 	{
-		return nullptr;
+		return MakeError(TEXTVIEW("Memory mapped IO is not enabled on this platform"));
 	}
 
 #if !UE_BUILD_SHIPPING
 	if (bLookLooseFirst && IsNonPakFilenameAllowed(Filename))
 	{
-		IMappedFileHandle* Handle = LowerLevel->OpenMapped(Filename, OpenOptions);
-		if (Handle != nullptr)
+		FOpenMappedResult Result = LowerLevel->OpenMappedEx(Filename, OpenOptions);
+		if (!Result.HasError())
 		{
-			return Handle;
+			return MakeValue(Result.StealValue());
 		}
 	}
 #endif
@@ -4615,25 +4614,27 @@ IMappedFileHandle* FPakPlatformFile::OpenMapped(const TCHAR* Filename, EOpenRead
 	{
 		if (FileEntry.CompressionMethodIndex != 0 || (FileEntry.Flags & FPakEntry::Flag_Encrypted) != 0)
 		{
-			// can't map compressed or encrypted files
-			return nullptr;
+			return MakeError(FString::Printf(TEXT("Can't map compressed or encrypted file '%s'"), Filename));
 		}
 		FScopeLock Lock(&PakEntry->MappedFileHandleCriticalSection);
 		if (!PakEntry->MappedFileHandle)
 		{
-			PakEntry->MappedFileHandle = LowerLevel->OpenMapped(*PakEntry->GetFilename(), OpenOptions);
-		}
-		if (!PakEntry->MappedFileHandle)
-		{
-			return nullptr;
+			FOpenMappedResult Result = LowerLevel->OpenMappedEx(*PakEntry->GetFilename(), OpenOptions);
+			if (Result.HasError())
+			{
+				return MakeError(Result.GetError().GetMessage());
+			}
+
+			PakEntry->MappedFileHandle = Result.StealValue().Get();
 		}
-		return new FMappedFilePakProxy(PakEntry->MappedFileHandle, FileEntry.Offset + FileEntry.GetSerializedSize(PakEntry->GetInfo().Version), FileEntry.UncompressedSize, PakEntry->TotalSize(), Filename);
+
+		return MakeValue(new FMappedFilePakProxy(PakEntry->MappedFileHandle, FileEntry.Offset + FileEntry.GetSerializedSize(PakEntry->GetInfo().Version), FileEntry.UncompressedSize, PakEntry->TotalSize(), Filename));
 	}
 	if (IsNonPakFilenameAllowed(Filename))
 	{
-		return LowerLevel->OpenMapped(Filename, OpenOptions);
+		return LowerLevel->OpenMappedEx(Filename, OpenOptions);
 	}
-	return nullptr;
+	return MakeError(FString::Printf(TEXT("Failed to open mapped file '%s'"), Filename));
 }
 
 IAsyncReadFileHandle* FPakPlatformFile::OpenAsyncRead(const TCHAR* Filename, bool bAllowWrite /*= false*/)
diff --git a/Engine/Source/Runtime/PakFile/Private/IoDispatcherFileBackend.cpp b/Engine/Source/Runtime/PakFile/Private/IoDispatcherFileBackend.cpp
--- a/Engine/Source/Runtime/PakFile/Private/IoDispatcherFileBackend.cpp
+++ b/Engine/Source/Runtime/PakFile/Private/IoDispatcherFileBackend.cpp
@@ -928,7 +928,8 @@ IMappedFileHandle* FFileIoStoreReader::GetMappedContainerFileHandle(uint64 TocOf
 		!EnumHasAnyFlags(ContainerFile.ContainerFlags, EIoContainerFlags::Encrypted))
 	{
 		IPlatformFile& Ipf = FPlatformFileManager::Get().GetPlatformFile();
-		Partition.MappedFileHandle.Reset(Ipf.OpenMapped(*Partition.FilePath));
+		FOpenMappedResult Result = Ipf.OpenMappedEx(*Partition.FilePath);
+		Partition.MappedFileHandle = Result.HasError() ? nullptr : Result.StealValue();
 	}
 
 	check(Partition.FileSize > 0);
diff --git a/Engine/Source/Runtime/PakFile/Public/IPlatformFilePak.h b/Engine/Source/Runtime/PakFile/Public/IPlatformFilePak.h
--- a/Engine/Source/Runtime/PakFile/Public/IPlatformFilePak.h
+++ b/Engine/Source/Runtime/PakFile/Public/IPlatformFilePak.h
@@ -2504,7 +2504,7 @@ public:
 	PAKFILE_API virtual IAsyncReadFileHandle* OpenAsyncRead(const TCHAR* Filename, bool bAllowWrite = false) override;
 	PAKFILE_API virtual void SetAsyncMinimumPriority(EAsyncIOPriorityAndFlags Priority) override;
 
-	virtual IMappedFileHandle* OpenMapped(const TCHAR* Filename, EOpenReadFlags OpenOptions = EOpenReadFlags::None) override;
+	virtual FOpenMappedResult OpenMappedEx(const TCHAR* Filename, EOpenReadFlags OpenOptions = EOpenReadFlags::None) override;
 	/**
 	 * Converts a filename to a path inside pak file.
 	 *
diff --git a/Engine/Source/Runtime/SandboxFile/Private/IPlatformFileSandboxWrapper.cpp b/Engine/Source/Runtime/SandboxFile/Private/IPlatformFileSandboxWrapper.cpp
--- a/Engine/Source/Runtime/SandboxFile/Private/IPlatformFileSandboxWrapper.cpp
+++ b/Engine/Source/Runtime/SandboxFile/Private/IPlatformFileSandboxWrapper.cpp
@@ -9,6 +9,7 @@
 #include "Modules/ModuleManager.h"
 #include "HAL/IPlatformFileModule.h"
 #include "Templates/UniquePtr.h"
+#include "Async/MappedFileHandle.h"
 
 DEFINE_LOG_CATEGORY(SandboxFile);
 
@@ -979,18 +980,30 @@ IAsyncReadFileHandle* FSandboxPlatformFile::OpenAsyncRead(const TCHAR* Filename,
 	}
 	return LowerLevel->OpenAsyncRead(Filename, bAllowWrite);
 }
+
 void FSandboxPlatformFile::SetAsyncMinimumPriority(EAsyncIOPriorityAndFlags Priority)
 {
 	LowerLevel->SetAsyncMinimumPriority(Priority);
 }
-IMappedFileHandle* FSandboxPlatformFile::OpenMapped(const TCHAR* Filename, EOpenReadFlags OpenOptions)
+
+IMappedFileHandle* FSandboxPlatformFile::OpenMapped(const TCHAR* Filename)
+{
+	FOpenMappedResult Result = OpenMappedEx(Filename);
+	if (Result.HasError())
+	{
+		return nullptr;
+	}
+	return Result.StealValue().Get();
+}
+
+FOpenMappedResult FSandboxPlatformFile::OpenMappedEx(const TCHAR* Filename, EOpenReadFlags OpenOptions)
 {
 	FString UserFilename(*ConvertToSandboxPath(Filename));
 	if (!OkForInnerAccess(Filename) || LowerLevel->FileExists(*UserFilename))
 	{
-		return LowerLevel->OpenMapped(*UserFilename, OpenOptions);
+		return LowerLevel->OpenMappedEx(*UserFilename, OpenOptions);
 	}
-	return LowerLevel->OpenMapped(Filename, OpenOptions);
+	return LowerLevel->OpenMappedEx(Filename, OpenOptions);
 }
 
 /**
diff --git a/Engine/Source/Runtime/SandboxFile/Public/IPlatformFileSandboxWrapper.h b/Engine/Source/Runtime/SandboxFile/Public/IPlatformFileSandboxWrapper.h
--- a/Engine/Source/Runtime/SandboxFile/Public/IPlatformFileSandboxWrapper.h
+++ b/Engine/Source/Runtime/SandboxFile/Public/IPlatformFileSandboxWrapper.h
@@ -223,7 +223,10 @@ public:
 
 	SANDBOXFILE_API virtual IAsyncReadFileHandle*	OpenAsyncRead(const TCHAR* Filename, bool bAllowWrite = false) override;
 	SANDBOXFILE_API virtual void					SetAsyncMinimumPriority(EAsyncIOPriorityAndFlags Priority) override;
-	SANDBOXFILE_API virtual IMappedFileHandle*		OpenMapped(const TCHAR* Filename, EOpenReadFlags OpenOptions = EOpenReadFlags::None) override;
+
+	UE_DEPRECATED(5.6, "Please use overloaded value that returns error values properly")
+	SANDBOXFILE_API virtual IMappedFileHandle*		OpenMapped(const TCHAR* Filename) override;
+	SANDBOXFILE_API virtual FOpenMappedResult		OpenMappedEx(const TCHAR* Filename, EOpenReadFlags OpenOptions = EOpenReadFlags::None) override;
 
 	friend class FSandboxVisitor;
 	friend class FSandboxStatVisitor;
diff --git a/Engine/Source/Runtime/StorageServerClient/Private/StorageServerPlatformFile.cpp b/Engine/Source/Runtime/StorageServerClient/Private/StorageServerPlatformFile.cpp
--- a/Engine/Source/Runtime/StorageServerClient/Private/StorageServerPlatformFile.cpp
+++ b/Engine/Source/Runtime/StorageServerClient/Private/StorageServerPlatformFile.cpp
@@ -25,6 +25,7 @@
 #include "StorageServerPackageStore.h"
 #include "Containers/LruCache.h"
 #include <atomic>
+#include "Async/MappedFileHandle.h"
 
 DEFINE_LOG_CATEGORY_STATIC(LogStorageServerPlatformFile, Log, All);
 
@@ -1096,9 +1097,13 @@ bool FStorageServerPlatformFile::IterateDirectoryStat(const TCHAR* Directory, FD
 	return bResult;
 }
 
-IMappedFileHandle* FStorageServerPlatformFile::OpenMapped(const TCHAR* Filename, EOpenReadFlags OpenOptions)
+FOpenMappedResult FStorageServerPlatformFile::OpenMappedEx(const TCHAR* Filename, EOpenReadFlags OpenOptions)
 {
-	return (LowerLevel && IsNonServerFilenameAllowed(Filename)) ? LowerLevel->OpenMapped(Filename, OpenOptions) : nullptr;
+	if (LowerLevel && IsNonServerFilenameAllowed(Filename))
+	{
+		return LowerLevel->OpenMappedEx(Filename, OpenOptions);
+	}
+	return MakeError(FString::Printf(TEXT("Can't open mapped file '%s'"), Filename));
 }
 
 bool FStorageServerPlatformFile::DirectoryExists(const TCHAR* Directory)
diff --git a/Engine/Source/Runtime/StorageServerClient/Private/StorageServerPlatformFile.h b/Engine/Source/Runtime/StorageServerClient/Private/StorageServerPlatformFile.h
--- a/Engine/Source/Runtime/StorageServerClient/Private/StorageServerPlatformFile.h
+++ b/Engine/Source/Runtime/StorageServerClient/Private/StorageServerPlatformFile.h
@@ -97,7 +97,7 @@ public:
 	virtual bool IterateDirectory(const TCHAR* Directory, FDirectoryVisitor& Visitor) override;
 	virtual bool IterateDirectoryRecursively(const TCHAR* Directory, FDirectoryVisitor& Visitor) override;
 	virtual bool IterateDirectoryStat(const TCHAR* Directory, FDirectoryStatVisitor& Visitor) override;
-	virtual IMappedFileHandle* OpenMapped(const TCHAR* Filename, EOpenReadFlags OpenOptions = EOpenReadFlags::None) override;
+	virtual FOpenMappedResult OpenMappedEx(const TCHAR* Filename, EOpenReadFlags OpenOptions = EOpenReadFlags::None) override;
 	virtual FString GetFilenameOnDisk(const TCHAR* Filename) override;
 	virtual bool DeleteFile(const TCHAR* Filename) override;
 	virtual bool MoveFile(const TCHAR* To, const TCHAR* From) override;
diff --git a/Engine/Source/Runtime/VirtualFileCache/Private/VirtualFileCacheInternal.cpp b/Engine/Source/Runtime/VirtualFileCache/Private/VirtualFileCacheInternal.cpp
--- a/Engine/Source/Runtime/VirtualFileCache/Private/VirtualFileCacheInternal.cpp
+++ b/Engine/Source/Runtime/VirtualFileCache/Private/VirtualFileCacheInternal.cpp
@@ -1303,9 +1303,10 @@ IMappedFileHandle* FFileTable::MapBlockFile(FBlockFile& BlockFile)
 		FRWScopeLock ScopeLock(*BlockFile.FileHandleLock.Get(), SLT_Write);
 		BlockFile.WriteHandle = nullptr;
 
-		// Open for read and write
+		// Open for read and write. Note - this doesn't open a file for write, for read only. Figure it out
 		IPlatformFile& PlatformFile = IPlatformFile::GetPlatformPhysical();
-		BlockFile.MapHandle = TUniquePtr<IMappedFileHandle>(PlatformFile.OpenMapped(*Path));
+		FOpenMappedResult Result = PlatformFile.OpenMappedEx(*Path);
+		BlockFile.MapHandle = Result.HasError() ? nullptr : Result.StealValue();
 	}
 	return BlockFile.MapHandle.Get();
 }
diff --git a/Engine/Source/Runtime/VulkanRHI/Private/VulkanChunkedPipelineCache.cpp b/Engine/Source/Runtime/VulkanRHI/Private/VulkanChunkedPipelineCache.cpp
--- a/Engine/Source/Runtime/VulkanRHI/Private/VulkanChunkedPipelineCache.cpp
+++ b/Engine/Source/Runtime/VulkanRHI/Private/VulkanChunkedPipelineCache.cpp
@@ -307,7 +307,8 @@ public:
 
 		if (UE::Vulkan::CanMemoryMapChunkedPSOCache())
 		{
-			MappedCacheFile = TUniquePtr<IMappedFileHandle>(FPlatformFileManager::Get().GetPlatformFile().OpenMapped(*FullCachePath));
+			FOpenMappedResult Result = FPlatformFileManager::Get().GetPlatformFile().OpenMappedEx(*FullCachePath);
+			MappedCacheFile = Result.HasError() ? nullptr : Result.StealValue();
 			UpdateMapping(LastValidOffset);
 		}
 
@@ -364,7 +365,8 @@ public:
 
 		if (UE::Vulkan::CanMemoryMapChunkedPSOCache() && !MappedCacheFile.IsValid())
 		{
-			MappedCacheFile = TUniquePtr<IMappedFileHandle>(FPlatformFileManager::Get().GetPlatformFile().OpenMapped(*FullCachePath));
+			FOpenMappedResult Result = FPlatformFileManager::Get().GetPlatformFile().OpenMappedEx(*FullCachePath);
+			MappedCacheFile = Result.HasError() ? nullptr : Result.StealValue();
 		}
 
 		return bWriteSuccess;
