Fixed incorrect use of UniquePtrs in file OpenMappedEx that results in releasing the file mapping object that's being returned

diff --git a/Engine/Source/Runtime/Core/Private/GenericPlatform/GenericPlatformFile.cpp b/Engine/Source/Runtime/Core/Private/GenericPlatform/GenericPlatformFile.cpp
--- a/Engine/Source/Runtime/Core/Private/GenericPlatform/GenericPlatformFile.cpp
+++ b/Engine/Source/Runtime/Core/Private/GenericPlatform/GenericPlatformFile.cpp
@@ -545,7 +545,7 @@ IMappedFileHandle* IPlatformFile::OpenMapped(const TCHAR* Filename)
 		return nullptr;
 	}
 
-	return Result.StealValue().Get();
+	return Result.StealValue().Release();
 }
 
 FOpenMappedResult IPlatformFile::OpenMappedEx(const TCHAR* Filename, EOpenReadFlags OpenOptions)
diff --git a/Engine/Source/Runtime/CoreUObject/Private/UObject/PackageResourceManagerFile.cpp b/Engine/Source/Runtime/CoreUObject/Private/UObject/PackageResourceManagerFile.cpp
--- a/Engine/Source/Runtime/CoreUObject/Private/UObject/PackageResourceManagerFile.cpp
+++ b/Engine/Source/Runtime/CoreUObject/Private/UObject/PackageResourceManagerFile.cpp
@@ -294,7 +294,7 @@ IMappedFileHandle* FPackageResourceManagerFile::OpenMappedHandleToPackage(const
 			{
 				return false;
 			}
-			Result = Value.StealValue().Get();
+			Result = Value.StealValue().Release();
 			return true;
 		});
 	return Result;
diff --git a/Engine/Source/Runtime/PakFile/Private/IPlatformFilePak.cpp b/Engine/Source/Runtime/PakFile/Private/IPlatformFilePak.cpp
--- a/Engine/Source/Runtime/PakFile/Private/IPlatformFilePak.cpp
+++ b/Engine/Source/Runtime/PakFile/Private/IPlatformFilePak.cpp
@@ -4625,7 +4625,7 @@ FOpenMappedResult FPakPlatformFile::OpenMappedEx(const TCHAR* Filename, EOpenRea
 				return MakeError(Result.GetError().GetMessage());
 			}
 
-			PakEntry->MappedFileHandle = Result.StealValue().Get();
+			PakEntry->MappedFileHandle = Result.StealValue().Release();
 		}
 
 		return MakeValue(new FMappedFilePakProxy(PakEntry->MappedFileHandle, FileEntry.Offset + FileEntry.GetSerializedSize(PakEntry->GetInfo().Version), FileEntry.UncompressedSize, PakEntry->TotalSize(), Filename));
diff --git a/Engine/Source/Runtime/SandboxFile/Private/IPlatformFileSandboxWrapper.cpp b/Engine/Source/Runtime/SandboxFile/Private/IPlatformFileSandboxWrapper.cpp
--- a/Engine/Source/Runtime/SandboxFile/Private/IPlatformFileSandboxWrapper.cpp
+++ b/Engine/Source/Runtime/SandboxFile/Private/IPlatformFileSandboxWrapper.cpp
@@ -993,7 +993,7 @@ IMappedFileHandle* FSandboxPlatformFile::OpenMapped(const TCHAR* Filename)
 	{
 		return nullptr;
 	}
-	return Result.StealValue().Get();
+	return Result.StealValue().Release();
 }
 
 FOpenMappedResult FSandboxPlatformFile::OpenMappedEx(const TCHAR* Filename, EOpenReadFlags OpenOptions)
