FPakFile mapped file handle ownership fix.

MappedFileHandle is now wrapped in a shared pointer, as it is actually shared between callers who OpenMapped on a file.
As FPakFiles are cached and same handle is being handed away - we can only free it, when all the users are gone.

Without this, there was a bunch of crashes happening due to:
- code OpenMapped a file from pak
- code gets few mapped regions from that file
- pak is being unloaded, with handle being freed, while regions still alive
- regions (e.g. FIOSMappedFileRegion) have Parent pointing to garbage, calling UnMap() on it with `NumOutstandingRegions--;` stomping the memory

r.Metal.EnableMetalDeferredDeleteLatency is now disabled by default as it was used as a speculative mitigation for one of the possible crashes.
MetalBuffer changes are also removing a mitigation for one of the crashes.

#rnx

diff --git a/Engine/Source/Runtime/Apple/MetalRHI/Private/MetalBuffer.cpp b/Engine/Source/Runtime/Apple/MetalRHI/Private/MetalBuffer.cpp
--- a/Engine/Source/Runtime/Apple/MetalRHI/Private/MetalBuffer.cpp
+++ b/Engine/Source/Runtime/Apple/MetalRHI/Private/MetalBuffer.cpp
@@ -71,7 +71,7 @@ void FMetalBuffer::Release()
 {
     check(!GIsRHIInitialized || (OnFreePolicy != FreePolicy::Owner || bMarkedDeleted));
     
-	if (OnFreePolicy == FreePolicy::BufferAllocator || Allocator)
+	if (OnFreePolicy == FreePolicy::BufferAllocator)
 	{
 		check(Allocator);
 		Allocator->ReleaseBuffer(this);
diff --git a/Engine/Source/Runtime/Apple/MetalRHI/Private/MetalRHI.cpp b/Engine/Source/Runtime/Apple/MetalRHI/Private/MetalRHI.cpp
--- a/Engine/Source/Runtime/Apple/MetalRHI/Private/MetalRHI.cpp
+++ b/Engine/Source/Runtime/Apple/MetalRHI/Private/MetalRHI.cpp
@@ -57,7 +57,7 @@ static TAutoConsoleVariable<bool> CVarEnableMetalPSOFileCacheWhenPrecachingActiv
 
 static TAutoConsoleVariable<bool> CVarEnableMetalDeferredDeleteLatency(
 	TEXT("r.Metal.EnableMetalDeferredDeleteLatency"),
-	true,
+	false,
 	TEXT("false: No added latency on deferred delete \n")
 	TEXT("true: Extra latency on deferred delete"),
 	ECVF_RenderThreadSafe | ECVF_ReadOnly);
diff --git a/Engine/Source/Runtime/Core/Private/IOS/IOSPlatformFile.cpp b/Engine/Source/Runtime/Core/Private/IOS/IOSPlatformFile.cpp
--- a/Engine/Source/Runtime/Core/Private/IOS/IOSPlatformFile.cpp
+++ b/Engine/Source/Runtime/Core/Private/IOS/IOSPlatformFile.cpp
@@ -507,7 +507,8 @@ public:
 
 	~FIOSMappedFileHandle()
 	{
-		check(!NumOutstandingRegions); // can't delete the file before you delete all outstanding regions
+		// can't delete the file before you delete all outstanding regions
+		UE_CLOG(NumOutstandingRegions, LogIOS, Fatal, TEXT("Cleaning mapped file with alive mapped regions."));
 		close(FileHandle);
 	}
 	
diff --git a/Engine/Source/Runtime/PakFile/Private/IPlatformFilePak.cpp b/Engine/Source/Runtime/PakFile/Private/IPlatformFilePak.cpp
--- a/Engine/Source/Runtime/PakFile/Private/IPlatformFilePak.cpp
+++ b/Engine/Source/Runtime/PakFile/Private/IPlatformFilePak.cpp
@@ -3905,24 +3905,20 @@ public:
 
 class FMappedFilePakProxy final : public IMappedFileHandle
 {
-	IMappedFileHandle* LowerLevel;
+	TSharedPtr<IMappedFileHandle> LowerLevel;
 	int64 OffsetInPak;
 	int64 PakSize;
 	FString DebugFilename;
 public:
-	FMappedFilePakProxy(IMappedFileHandle* InLowerLevel, int64 InOffset, int64 InSize, int64 InPakSize, const TCHAR* InDebugFilename)
+	FMappedFilePakProxy(TSharedPtr<IMappedFileHandle> InLowerLevel, int64 InOffset, int64 InSize, int64 InPakSize, const TCHAR* InDebugFilename)
 		: IMappedFileHandle(InSize)
-		, LowerLevel(InLowerLevel)
+		, LowerLevel(MoveTemp(InLowerLevel))
 		, OffsetInPak(InOffset)
 		, PakSize(InPakSize)
 		, DebugFilename(InDebugFilename)
 	{
 		check(PakSize >= 0);
 	}
-	virtual ~FMappedFilePakProxy()
-	{
-		// we don't own lower level, it is shared
-	}
 	virtual IMappedFileRegion* MapRegion(int64 Offset = 0, int64 BytesToMap = MAX_int64, FFileMappingFlags Flags = EMappedFileFlags::ENone) override
 	{
 		//check(Offset + OffsetInPak < PakSize); // don't map zero bytes and don't map off the end of the (real) file
@@ -4625,7 +4621,7 @@ FOpenMappedResult FPakPlatformFile::OpenMappedEx(const TCHAR* Filename, EOpenRea
 				return MakeError(Result.StealError());
 			}
 
-			PakEntry->MappedFileHandle = Result.StealValue().Release();
+			PakEntry->MappedFileHandle = TSharedPtr<IMappedFileHandle>(Result.StealValue().Release());
 		}
 
 		return MakeValue(new FMappedFilePakProxy(PakEntry->MappedFileHandle, FileEntry.Offset + FileEntry.GetSerializedSize(PakEntry->GetInfo().Version), FileEntry.UncompressedSize, PakEntry->TotalSize(), Filename));
diff --git a/Engine/Source/Runtime/PakFile/Private/PakFile.cpp b/Engine/Source/Runtime/PakFile/Private/PakFile.cpp
--- a/Engine/Source/Runtime/PakFile/Private/PakFile.cpp
+++ b/Engine/Source/Runtime/PakFile/Private/PakFile.cpp
@@ -186,7 +186,6 @@ FPakFile::FPakFile(FArchive* Archive)
 
 FPakFile::~FPakFile()
 {
-	delete MappedFileHandle;
 }
 
 bool FPakFile::PassedSignatureChecks() const
diff --git a/Engine/Source/Runtime/PakFile/Public/IPlatformFilePak.h b/Engine/Source/Runtime/PakFile/Public/IPlatformFilePak.h
--- a/Engine/Source/Runtime/PakFile/Public/IPlatformFilePak.h
+++ b/Engine/Source/Runtime/PakFile/Public/IPlatformFilePak.h
@@ -832,7 +832,7 @@ private:
 	/** ID for the chunk this pakfile is part of. INDEX_NONE if this isn't a pak chunk (derived from filename) */
 	int32 PakchunkIndex;
 
-	class IMappedFileHandle* MappedFileHandle;
+	TSharedPtr<IMappedFileHandle> MappedFileHandle;
 	FCriticalSection MappedFileHandleCriticalSection;
 
 	/** The type of cache this pak file should have */
