Support for dynamic I/O store signature block size.

#rnx

diff --git a/Engine/Source/Runtime/PakFile/Internal/IoDispatcherFileBackendTypes.h b/Engine/Source/Runtime/PakFile/Internal/IoDispatcherFileBackendTypes.h
--- a/Engine/Source/Runtime/PakFile/Internal/IoDispatcherFileBackendTypes.h
+++ b/Engine/Source/Runtime/PakFile/Internal/IoDispatcherFileBackendTypes.h
@@ -40,27 +40,55 @@ struct FFileIoStoreContainerFilePartition
 	std::atomic<int32> StartedReadRequestsCount = 0;
 };
 
-struct FFileIoStoreBlockSignatureTable
+class FFileIoStoreBlockSignatures
 {
-	TArray<FSHAHash> Hashes;
+	static constexpr uint32 MinBlockSize = 4;
 
-	void AddRef()
+public:
+	bool IsValidBlockSignature(const FSHAHash& Signature, uint32 BlockIndex, FSHAHash& OutExpectedSignature) const
 	{
-		++RefCount;
-	}
+		const uint32 Offset = BlockIndex * BlockSize;
+		ensure(Data.IsValid() && Offset <= DataSize - BlockSize);
 
-	void Release()
-	{
-		if (--RefCount == 0)
+		if (FMemory::Memcmp(Data.Get() + Offset, &Signature, BlockSize) != 0)
 		{
-			delete this;
+			FMemory::Memcpy(&OutExpectedSignature, Data.Get() + Offset, BlockSize);
+			return false;
 		}
+
+		return true;
+	}
+
+	bool IsValidBlockSignature(const FSHAHash& Signature, uint32 BlockIndex) const
+	{
+		FSHAHash Tmp;
+		return IsValidBlockSignature(Signature, BlockIndex, Tmp);
+	}
+
+	uint32 GetAllocatedSize() const { return DataSize; }
+
+	static TSharedPtr<FFileIoStoreBlockSignatures, ESPMode::ThreadSafe> Create(TConstArrayView<FSHAHash> BlockSignatures, uint32 BlockSize);
+
+	static uint32 GetValidSignatureBlockSize(uint32 BlockSize)
+	{
+		BlockSize += BlockSize % 2;
+		return FMath::Max(MinBlockSize, FMath::Min(BlockSize, static_cast<uint32>(sizeof(FSHAHash))));
 	}
 
 private:
-	std::atomic<int64> RefCount = 0;
+	FFileIoStoreBlockSignatures(TUniquePtr<uint8[]>&& InData, uint32 InDataSize, uint32 InBlockSize)
+		: Data(MoveTemp(InData))
+		, DataSize(InDataSize)
+		, BlockSize(InBlockSize)
+	{ }
+
+	TUniquePtr<uint8[]> Data;
+	uint32				DataSize = 0;
+	uint32				BlockSize = 0;
 };
 
+using FSharedBlockSignatures = TSharedPtr<FFileIoStoreBlockSignatures, ESPMode::ThreadSafe>;
+
 struct FFileIoStoreContainerFile
 {
 	FFileIoStoreContainerFile() = default;
@@ -78,7 +106,7 @@ struct FFileIoStoreContainerFile
 	FGuid EncryptionKeyGuid;
 	FAES::FAESKey EncryptionKey;
 	EIoContainerFlags ContainerFlags;
-	TRefCountPtr<FFileIoStoreBlockSignatureTable> BlockSignatureTable;
+	FSharedBlockSignatures BlockSignatures;
 	TArray<FFileIoStoreContainerFilePartition> Partitions;
 	uint32 ContainerInstanceId = 0;
 
@@ -144,8 +172,7 @@ struct FFileIoStoreCompressedBlock
 	FFileIoStoreCompressionContext* CompressionContext = nullptr;
 	uint8* CompressedDataBuffer = nullptr;
 	FAES::FAESKey EncryptionKey;
-	TRefCountPtr<FFileIoStoreBlockSignatureTable> BlockSignatureTable;
-	const FSHAHash* SignatureHash = nullptr;
+	FSharedBlockSignatures BlockSignatures;
 	bool bFailed = false;
 	bool bCancelled = false;
 };
diff --git a/Engine/Source/Runtime/PakFile/Private/IoDispatcherFileBackend.cpp b/Engine/Source/Runtime/PakFile/Private/IoDispatcherFileBackend.cpp
--- a/Engine/Source/Runtime/PakFile/Private/IoDispatcherFileBackend.cpp
+++ b/Engine/Source/Runtime/PakFile/Private/IoDispatcherFileBackend.cpp
@@ -114,6 +114,16 @@ static FAutoConsoleVariableRef CVar_IoDispatcherForceSynchronousScatter(
 	TEXT("This can avoid deadlocks in cases where background tasks end up waiting on I/O and we don't have enough background task threads to fulfill decompression requests.")
 );
 
+#ifndef IO_DISPATCHER_SIGNATURE_BLOCK_SIZE
+int32 GIoDispatcherSignatureBlockSize = sizeof(FSHAHash);
+#else
+int32 GIoDispatcherSignatureBlockSize = IO_DISPATCHER_SIGNATURE_BLOCK_SIZE;
+#endif
+static FAutoConsoleVariableRef CVar_IoDispatcherSignatureBlockSize(
+	TEXT("s.IoDispatcherSignatureBlockSize"),
+	GIoDispatcherSignatureBlockSize,
+	TEXT("Signature block size")
+);
 
 uint32 FFileIoStoreReadRequest::NextSequence = 0;
 #if CHECK_IO_STORE_READ_REQUEST_LIST_MEMBERSHIP
@@ -122,6 +132,32 @@ uint32 FFileIoStoreReadRequestList::NextListCookie = 0;
 TAtomic<uint32> FFileIoStoreReader::GlobalPartitionIndex{ 0 };
 TAtomic<uint32> FFileIoStoreReader::GlobalContainerInstanceId{ 0 };
 
+FSharedBlockSignatures FFileIoStoreBlockSignatures::Create(
+	TConstArrayView<FSHAHash> BlockSignatures,
+	uint32 BlockSize)
+{
+	ensure(BlockSize > 0 && BlockSize <= sizeof(FSHAHash) && ((BlockSize & 1) == 0));
+
+	if (BlockSignatures.IsEmpty())
+	{
+		return FSharedBlockSignatures(new FFileIoStoreBlockSignatures(TUniquePtr<uint8[]>(), 0, 0));
+	}
+
+	const uint32 DataSize = BlockSignatures.Num() * BlockSize;
+	TUniquePtr<uint8[]> Data = MakeUnique<uint8[]>(DataSize);
+	FMutableMemoryView DataView = MakeMemoryView(Data.Get(), DataSize);
+
+	for (const FSHAHash& Signature : BlockSignatures)
+	{
+		FMemoryView SignatureView = MakeMemoryView(&Signature, BlockSize);
+		DataView.CopyFrom(SignatureView);
+		DataView += BlockSize;
+	}
+	check(DataView.IsEmpty());
+
+	return FSharedBlockSignatures(new FFileIoStoreBlockSignatures(MoveTemp(Data), DataSize, BlockSize));
+}
+
 class FMappedFileProxy final : public IMappedFileHandle
 {
 public:
@@ -731,7 +767,7 @@ uint64 FFileIoStoreReader::GetTocAllocatedSize() const
 		PerfectHashMap.TocChunkIds.GetAllocatedSize() +
 		PerfectHashMap.TocChunkHashSeeds.GetAllocatedSize() +
 		ContainerFile.CompressionBlocks.GetAllocatedSize() +
-		(ContainerFile.BlockSignatureTable.IsValid() ? ContainerFile.BlockSignatureTable->Hashes.GetAllocatedSize() : 0);
+		(ContainerFile.BlockSignatures.IsValid() ? ContainerFile.BlockSignatures->GetAllocatedSize() : 0);
 }
 
 FIoStatus FFileIoStoreReader::Initialize(const TCHAR* InTocFilePath, int32 InOrder)
@@ -805,14 +841,19 @@ FIoStatus FFileIoStoreReader::Initialize(const TCHAR* InTocFilePath, int32 InOrd
 	ContainerFile.CompressionBlocks		= MoveTemp(TocResource.CompressionBlocks);
 	ContainerFile.ContainerFlags		= TocResource.Header.ContainerFlags;
 	ContainerFile.EncryptionKeyGuid		= TocResource.Header.EncryptionKeyGuid;
-	ContainerFile.BlockSignatureTable	= new FFileIoStoreBlockSignatureTable();
-	ContainerFile.BlockSignatureTable->Hashes = MoveTemp(TocResource.ChunkBlockSignatures);
 	ContainerFile.ContainerInstanceId	= ++GlobalContainerInstanceId;
 
+	const uint32 SignatureBlockSize = FFileIoStoreBlockSignatures::GetValidSignatureBlockSize(GIoDispatcherSignatureBlockSize);
+	if (EnumHasAnyFlags(ContainerFile.ContainerFlags, EIoContainerFlags::Signed))
+	{
+		ContainerFile.BlockSignatures = FFileIoStoreBlockSignatures::Create(TocResource.ChunkBlockSignatures, SignatureBlockSize);
+	}
+
 	Stats.OnTocMounted(GetTocAllocatedSize());
 
-	UE_LOG(LogIoDispatcher, Display, TEXT("Toc loaded : %s, Id=%s, Order=%d, EntryCount=%u, SignatureHash=%s"),
-		InTocFilePath, *LexToString(TocResource.Header.ContainerId), InOrder, TocResource.Header.TocEntryCount, *TocResource.SignatureHash.ToString());
+	UE_LOG(LogIoDispatcher, Display, TEXT("Toc loaded '%s', Id='%s', Order=%d, EntryCount=%u, SignatureHash='%s' SignatureBlockSize=%u"),
+		InTocFilePath, *LexToString(TocResource.Header.ContainerId), InOrder, TocResource.Header.TocEntryCount,
+		*TocResource.SignatureHash.ToString(), SignatureBlockSize);
 	
 	ContainerId = TocResource.Header.ContainerId;
 	Order = InOrder;
@@ -999,10 +1040,10 @@ TIoStatusOr<FIoContainerHeader> FFileIoStoreReader::ReadContainerHeader() const
 			const uint64 BlockSize = Align(CompressionBlockEntry->GetCompressedSize(), FAES::AESBlockSize);
 			if (bSigned)
 			{
-				const FSHAHash& SignatureHash = ContainerFile.BlockSignatureTable->Hashes[CompressedBlockIndex];
+				check(ContainerFile.BlockSignatures.IsValid());
 				FSHAHash BlockHash;
 				FSHA1::HashBuffer(BlockData, BlockSize, BlockHash.Hash);
-				if (SignatureHash != BlockHash)
+				if (!ContainerFile.BlockSignatures->IsValidBlockSignature(BlockHash, CompressedBlockIndex))
 				{
 					return FIoStatus(FIoStatusBuilder(EIoErrorCode::SignatureError) << TEXT("Signature error detected when reading container header"));
 				}
@@ -1623,11 +1664,11 @@ void FFileIoStore::ScatterBlock(FFileIoStoreCompressedBlock* CompressedBlock, bo
 	}
 	if (!CompressedBlock->bFailed)
 	{
-		if (CompressedBlock->SignatureHash)
+		if (CompressedBlock->BlockSignatures.IsValid())
 		{
-			FSHAHash BlockHash;
+			FSHAHash BlockHash, ExpectedBlockHash;
 			FSHA1::HashBuffer(CompressedBuffer, CompressedBlock->RawSize, BlockHash.Hash);
-			if (*CompressedBlock->SignatureHash != BlockHash)
+			if (!CompressedBlock->BlockSignatures->IsValidBlockSignature(BlockHash, CompressedBlock->Key.BlockIndex, ExpectedBlockHash))
 			{
 				FIoSignatureError Error;
 				{
@@ -1640,7 +1681,7 @@ void FFileIoStore::ScatterBlock(FFileIoStoreCompressedBlock* CompressedBlock, bo
 						}
 					}
 					Error.BlockIndex = CompressedBlock->Key.BlockIndex;
-					Error.ExpectedHash = *CompressedBlock->SignatureHash;
+					Error.ExpectedHash = ExpectedBlockHash;
 					Error.ActualHash = BlockHash;
 				}
 
@@ -1976,7 +2017,7 @@ FIoRequestImpl* FFileIoStore::GetCompletedIoRequests()
 		const bool bScatterAsync = bIsMultithreaded && GIoDispatcherForceSynchronousScatter == 0 &&
 			(!BlockToDecompress->CompressionMethod.IsNone() ||
 			 BlockToDecompress->EncryptionKey.IsValid() ||
-			 BlockToDecompress->SignatureHash) && 
+			 BlockToDecompress->BlockSignatures.IsValid()) && 
 			 // If we're already oversubscribed, we might not receive any further event to wake us and 
 			 // allow us to process our queue. In that case we simply run decompression locally.
 			 !FileIoStoreImpl::IsSchedulerOversubscribed(IoDispatcherTaskPriority);
@@ -2121,9 +2162,8 @@ void FFileIoStore::ReadBlocks(FFileIoStoreResolvedRequest& ResolvedRequest)
 			CompressedBlock->CompressionMethod = ContainerFile->CompressionMethods[CompressionBlockEntry.GetCompressionMethodIndex()];
 			if (EnumHasAnyFlags(ContainerFile->ContainerFlags, EIoContainerFlags::Signed))
 			{
-				check(ContainerFile->BlockSignatureTable);
-				CompressedBlock->BlockSignatureTable = ContainerFile->BlockSignatureTable;
-				CompressedBlock->SignatureHash = &ContainerFile->BlockSignatureTable->Hashes[CompressedBlockIndex];
+				check(ContainerFile->BlockSignatures.IsValid());
+				CompressedBlock->BlockSignatures = ContainerFile->BlockSignatures;
 			}
 			CompressedBlock->RawSize = Align(CompressionBlockEntry.GetCompressedSize(), FAES::AESBlockSize); // The raw blocks size is always aligned to AES blocks size;
 
