Fixed EntryData reserving 4x times more of elements than needed (using GetAllocatedSize instead of Num)

#rnx

diff --git a/Engine/Source/Runtime/PakFile/Private/FilePackageStore.cpp b/Engine/Source/Runtime/PakFile/Private/FilePackageStore.cpp
--- a/Engine/Source/Runtime/PakFile/Private/FilePackageStore.cpp
+++ b/Engine/Source/Runtime/PakFile/Private/FilePackageStore.cpp
@@ -593,7 +593,7 @@ void FFilePackageStoreBackend::Update()
 #endif
 		TArray<uint32> OldEntryData;
 		Swap(EntryData, OldEntryData);
-		EntryData.Reserve(static_cast<int32>(OldEntryData.GetAllocatedSize() + 32 * TotalNewPackages));
+		EntryData.Reserve(static_cast<int32>(OldEntryData.Num() + 32 * TotalNewPackages));
 	
 		// Helper constants
 		const uint32 IllegalHandleValue = ~0u;
