[Backout] - CL38441992
Conflicts with upcoming MMIO utocs.

#rnx
Original CL Desc
-----------------------------------------------------------------
Support for dynamic I/O store signature block size.

#rnx

diff --git a/Engine/Source/Runtime/PakFile/Internal/IoDispatcherFileBackendTypes.h b/Engine/Source/Runtime/PakFile/Internal/IoDispatcherFileBackendTypes.h
--- a/Engine/Source/Runtime/PakFile/Internal/IoDispatcherFileBackendTypes.h
+++ b/Engine/Source/Runtime/PakFile/Internal/IoDispatcherFileBackendTypes.h
@@ -40,55 +40,27 @@ struct FFileIoStoreContainerFilePartition
 	std::atomic<int32> StartedReadRequestsCount = 0;
 };
 
-class FFileIoStoreBlockSignatures
+struct FFileIoStoreBlockSignatureTable
 {
-	static constexpr uint32 MinBlockSize = 4;
+	TArray<FSHAHash> Hashes;
 
-public:
-	bool IsValidBlockSignature(const FSHAHash& Signature, uint32 BlockIndex, FSHAHash& OutExpectedSignature) const
+	void AddRef()
 	{
-		const uint32 Offset = BlockIndex * BlockSize;
-		ensure(Data.IsValid() && Offset <= DataSize - BlockSize);
-
-		if (FMemory::Memcmp(Data.Get() + Offset, &Signature, BlockSize) != 0)
-		{
-			FMemory::Memcpy(&OutExpectedSignature, Data.Get() + Offset, BlockSize);
-			return false;
-		}
-
-		return true;
+		++RefCount;
 	}
 
-	bool IsValidBlockSignature(const FSHAHash& Signature, uint32 BlockIndex) const
+	void Release()
 	{
-		FSHAHash Tmp;
-		return IsValidBlockSignature(Signature, BlockIndex, Tmp);
-	}
-
-	uint32 GetAllocatedSize() const { return DataSize; }
-
-	static TSharedPtr<FFileIoStoreBlockSignatures, ESPMode::ThreadSafe> Create(TConstArrayView<FSHAHash> BlockSignatures, uint32 BlockSize);
-
-	static uint32 GetValidSignatureBlockSize(uint32 BlockSize)
-	{
-		BlockSize += BlockSize % 2;
-		return FMath::Max(MinBlockSize, FMath::Min(BlockSize, static_cast<uint32>(sizeof(FSHAHash))));
+		if (--RefCount == 0)
+		{
+			delete this;
+		}
 	}
 
 private:
-	FFileIoStoreBlockSignatures(TUniquePtr<uint8[]>&& InData, uint32 InDataSize, uint32 InBlockSize)
-		: Data(MoveTemp(InData))
-		, DataSize(InDataSize)
-		, BlockSize(InBlockSize)
-	{ }
-
-	TUniquePtr<uint8[]> Data;
-	uint32				DataSize = 0;
-	uint32				BlockSize = 0;
+	std::atomic<int64> RefCount = 0;
 };
 
-using FSharedBlockSignatures = TSharedPtr<FFileIoStoreBlockSignatures, ESPMode::ThreadSafe>;
-
 struct FFileIoStoreContainerFile
 {
 	FFileIoStoreContainerFile() = default;
@@ -106,7 +78,7 @@ struct FFileIoStoreContainerFile
 	FGuid EncryptionKeyGuid;
 	FAES::FAESKey EncryptionKey;
 	EIoContainerFlags ContainerFlags;
-	FSharedBlockSignatures BlockSignatures;
+	TRefCountPtr<FFileIoStoreBlockSignatureTable> BlockSignatureTable;
 	TArray<FFileIoStoreContainerFilePartition> Partitions;
 	uint32 ContainerInstanceId = 0;
 
@@ -172,7 +144,8 @@ struct FFileIoStoreCompressedBlock
 	FFileIoStoreCompressionContext* CompressionContext = nullptr;
 	uint8* CompressedDataBuffer = nullptr;
 	FAES::FAESKey EncryptionKey;
-	FSharedBlockSignatures BlockSignatures;
+	TRefCountPtr<FFileIoStoreBlockSignatureTable> BlockSignatureTable;
+	const FSHAHash* SignatureHash = nullptr;
 	bool bFailed = false;
 	bool bCancelled = false;
 };
diff --git a/Engine/Source/Runtime/PakFile/Private/IoDispatcherFileBackend.cpp b/Engine/Source/Runtime/PakFile/Private/IoDispatcherFileBackend.cpp
--- a/Engine/Source/Runtime/PakFile/Private/IoDispatcherFileBackend.cpp
+++ b/Engine/Source/Runtime/PakFile/Private/IoDispatcherFileBackend.cpp
@@ -114,16 +114,6 @@ static FAutoConsoleVariableRef CVar_IoDispatcherForceSynchronousScatter(
 	TEXT("This can avoid deadlocks in cases where background tasks end up waiting on I/O and we don't have enough background task threads to fulfill decompression requests.")
 );
 
-#ifndef IO_DISPATCHER_SIGNATURE_BLOCK_SIZE
-int32 GIoDispatcherSignatureBlockSize = sizeof(FSHAHash);
-#else
-int32 GIoDispatcherSignatureBlockSize = IO_DISPATCHER_SIGNATURE_BLOCK_SIZE;
-#endif
-static FAutoConsoleVariableRef CVar_IoDispatcherSignatureBlockSize(
-	TEXT("s.IoDispatcherSignatureBlockSize"),
-	GIoDispatcherSignatureBlockSize,
-	TEXT("Signature block size")
-);
 
 uint32 FFileIoStoreReadRequest::NextSequence = 0;
 #if CHECK_IO_STORE_READ_REQUEST_LIST_MEMBERSHIP
@@ -132,32 +122,6 @@ uint32 FFileIoStoreReadRequestList::NextListCookie = 0;
 TAtomic<uint32> FFileIoStoreReader::GlobalPartitionIndex{ 0 };
 TAtomic<uint32> FFileIoStoreReader::GlobalContainerInstanceId{ 0 };
 
-FSharedBlockSignatures FFileIoStoreBlockSignatures::Create(
-	TConstArrayView<FSHAHash> BlockSignatures,
-	uint32 BlockSize)
-{
-	ensure(BlockSize > 0 && BlockSize <= sizeof(FSHAHash) && ((BlockSize & 1) == 0));
-
-	if (BlockSignatures.IsEmpty())
-	{
-		return FSharedBlockSignatures(new FFileIoStoreBlockSignatures(TUniquePtr<uint8[]>(), 0, 0));
-	}
-
-	const uint32 DataSize = BlockSignatures.Num() * BlockSize;
-	TUniquePtr<uint8[]> Data = MakeUnique<uint8[]>(DataSize);
-	FMutableMemoryView DataView = MakeMemoryView(Data.Get(), DataSize);
-
-	for (const FSHAHash& Signature : BlockSignatures)
-	{
-		FMemoryView SignatureView = MakeMemoryView(&Signature, BlockSize);
-		DataView.CopyFrom(SignatureView);
-		DataView += BlockSize;
-	}
-	check(DataView.IsEmpty());
-
-	return FSharedBlockSignatures(new FFileIoStoreBlockSignatures(MoveTemp(Data), DataSize, BlockSize));
-}
-
 class FMappedFileProxy final : public IMappedFileHandle
 {
 public:
@@ -767,7 +731,7 @@ uint64 FFileIoStoreReader::GetTocAllocatedSize() const
 		PerfectHashMap.TocChunkIds.GetAllocatedSize() +
 		PerfectHashMap.TocChunkHashSeeds.GetAllocatedSize() +
 		ContainerFile.CompressionBlocks.GetAllocatedSize() +
-		(ContainerFile.BlockSignatures.IsValid() ? ContainerFile.BlockSignatures->GetAllocatedSize() : 0);
+		(ContainerFile.BlockSignatureTable.IsValid() ? ContainerFile.BlockSignatureTable->Hashes.GetAllocatedSize() : 0);
 }
 
 FIoStatus FFileIoStoreReader::Initialize(const TCHAR* InTocFilePath, int32 InOrder)
@@ -841,19 +805,14 @@ FIoStatus FFileIoStoreReader::Initialize(const TCHAR* InTocFilePath, int32 InOrd
 	ContainerFile.CompressionBlocks		= MoveTemp(TocResource.CompressionBlocks);
 	ContainerFile.ContainerFlags		= TocResource.Header.ContainerFlags;
 	ContainerFile.EncryptionKeyGuid		= TocResource.Header.EncryptionKeyGuid;
+	ContainerFile.BlockSignatureTable	= new FFileIoStoreBlockSignatureTable();
+	ContainerFile.BlockSignatureTable->Hashes = MoveTemp(TocResource.ChunkBlockSignatures);
 	ContainerFile.ContainerInstanceId	= ++GlobalContainerInstanceId;
 
-	const uint32 SignatureBlockSize = FFileIoStoreBlockSignatures::GetValidSignatureBlockSize(GIoDispatcherSignatureBlockSize);
-	if (EnumHasAnyFlags(ContainerFile.ContainerFlags, EIoContainerFlags::Signed))
-	{
-		ContainerFile.BlockSignatures = FFileIoStoreBlockSignatures::Create(TocResource.ChunkBlockSignatures, SignatureBlockSize);
-	}
-
 	Stats.OnTocMounted(GetTocAllocatedSize());
 
-	UE_LOG(LogIoDispatcher, Display, TEXT("Toc loaded '%s', Id='%s', Order=%d, EntryCount=%u, SignatureHash='%s' SignatureBlockSize=%u"),
-		InTocFilePath, *LexToString(TocResource.Header.ContainerId), InOrder, TocResource.Header.TocEntryCount,
-		*TocResource.SignatureHash.ToString(), SignatureBlockSize);
+	UE_LOG(LogIoDispatcher, Display, TEXT("Toc loaded : %s, Id=%s, Order=%d, EntryCount=%u, SignatureHash=%s"),
+		InTocFilePath, *LexToString(TocResource.Header.ContainerId), InOrder, TocResource.Header.TocEntryCount, *TocResource.SignatureHash.ToString());
 	
 	ContainerId = TocResource.Header.ContainerId;
 	Order = InOrder;
@@ -1040,10 +999,10 @@ TIoStatusOr<FIoContainerHeader> FFileIoStoreReader::ReadContainerHeader() const
 			const uint64 BlockSize = Align(CompressionBlockEntry->GetCompressedSize(), FAES::AESBlockSize);
 			if (bSigned)
 			{
-				check(ContainerFile.BlockSignatures.IsValid());
+				const FSHAHash& SignatureHash = ContainerFile.BlockSignatureTable->Hashes[CompressedBlockIndex];
 				FSHAHash BlockHash;
 				FSHA1::HashBuffer(BlockData, BlockSize, BlockHash.Hash);
-				if (!ContainerFile.BlockSignatures->IsValidBlockSignature(BlockHash, CompressedBlockIndex))
+				if (SignatureHash != BlockHash)
 				{
 					return FIoStatus(FIoStatusBuilder(EIoErrorCode::SignatureError) << TEXT("Signature error detected when reading container header"));
 				}
@@ -1664,11 +1623,11 @@ void FFileIoStore::ScatterBlock(FFileIoStoreCompressedBlock* CompressedBlock, bo
 	}
 	if (!CompressedBlock->bFailed)
 	{
-		if (CompressedBlock->BlockSignatures.IsValid())
+		if (CompressedBlock->SignatureHash)
 		{
-			FSHAHash BlockHash, ExpectedBlockHash;
+			FSHAHash BlockHash;
 			FSHA1::HashBuffer(CompressedBuffer, CompressedBlock->RawSize, BlockHash.Hash);
-			if (!CompressedBlock->BlockSignatures->IsValidBlockSignature(BlockHash, CompressedBlock->Key.BlockIndex, ExpectedBlockHash))
+			if (*CompressedBlock->SignatureHash != BlockHash)
 			{
 				FIoSignatureError Error;
 				{
@@ -1681,7 +1640,7 @@ void FFileIoStore::ScatterBlock(FFileIoStoreCompressedBlock* CompressedBlock, bo
 						}
 					}
 					Error.BlockIndex = CompressedBlock->Key.BlockIndex;
-					Error.ExpectedHash = ExpectedBlockHash;
+					Error.ExpectedHash = *CompressedBlock->SignatureHash;
 					Error.ActualHash = BlockHash;
 				}
 
@@ -2017,7 +1976,7 @@ FIoRequestImpl* FFileIoStore::GetCompletedIoRequests()
 		const bool bScatterAsync = bIsMultithreaded && GIoDispatcherForceSynchronousScatter == 0 &&
 			(!BlockToDecompress->CompressionMethod.IsNone() ||
 			 BlockToDecompress->EncryptionKey.IsValid() ||
-			 BlockToDecompress->BlockSignatures.IsValid()) && 
+			 BlockToDecompress->SignatureHash) && 
 			 // If we're already oversubscribed, we might not receive any further event to wake us and 
 			 // allow us to process our queue. In that case we simply run decompression locally.
 			 !FileIoStoreImpl::IsSchedulerOversubscribed(IoDispatcherTaskPriority);
@@ -2162,8 +2121,9 @@ void FFileIoStore::ReadBlocks(FFileIoStoreResolvedRequest& ResolvedRequest)
 			CompressedBlock->CompressionMethod = ContainerFile->CompressionMethods[CompressionBlockEntry.GetCompressionMethodIndex()];
 			if (EnumHasAnyFlags(ContainerFile->ContainerFlags, EIoContainerFlags::Signed))
 			{
-				check(ContainerFile->BlockSignatures.IsValid());
-				CompressedBlock->BlockSignatures = ContainerFile->BlockSignatures;
+				check(ContainerFile->BlockSignatureTable);
+				CompressedBlock->BlockSignatureTable = ContainerFile->BlockSignatureTable;
+				CompressedBlock->SignatureHash = &ContainerFile->BlockSignatureTable->Hashes[CompressedBlockIndex];
 			}
 			CompressedBlock->RawSize = Align(CompressionBlockEntry.GetCompressedSize(), FAES::AESBlockSize); // The raw blocks size is always aligned to AES blocks size;
 
