Allow skipping the mounting pak files matching a wildcard in non-shipping builds.

diff --git a/Engine/Source/Runtime/PakFile/Private/IPlatformFilePak.cpp b/Engine/Source/Runtime/PakFile/Private/IPlatformFilePak.cpp
--- a/Engine/Source/Runtime/PakFile/Private/IPlatformFilePak.cpp
+++ b/Engine/Source/Runtime/PakFile/Private/IPlatformFilePak.cpp
@@ -5437,13 +5437,15 @@ void FPakPlatformFile::FindPakFilesInDirectory(IPlatformFile* LowLevelFile, cons
 	class FPakSearchVisitor : public IPlatformFile::FDirectoryVisitor
 	{
 		TArray<FString>& FoundPakFiles;
-		FString WildCard;
+		const FString& WildCard;
+		const FString& SkipPaksWildCard;
 		bool bSkipOptionalPakFiles;
 
 	public:
-		FPakSearchVisitor(TArray<FString>& InFoundPakFiles, const FString& InWildCard, bool bInSkipOptionalPakFiles)
+		FPakSearchVisitor(TArray<FString>& InFoundPakFiles, const FString& InWildCard, const FString& InSkipPaksWildCard, bool bInSkipOptionalPakFiles)
 			: FoundPakFiles(InFoundPakFiles)
 			, WildCard(InWildCard)
+			, SkipPaksWildCard(InSkipPaksWildCard)
 			, bSkipOptionalPakFiles(bInSkipOptionalPakFiles)
 		{}
 		virtual bool Visit(const TCHAR* FilenameOrDirectory, bool bIsDirectory)
@@ -5451,7 +5453,7 @@ void FPakPlatformFile::FindPakFilesInDirectory(IPlatformFile* LowLevelFile, cons
 			if (bIsDirectory == false)
 			{
 				FString Filename(FilenameOrDirectory);
-				if(Filename.MatchesWildcard(WildCard))
+				if (Filename.MatchesWildcard(WildCard))
 				{
 					if (!FPakPlatformFile::IsPakFileInstalled(Filename))
 					{
@@ -5459,11 +5461,18 @@ void FPakPlatformFile::FindPakFilesInDirectory(IPlatformFile* LowLevelFile, cons
 					}
 
 #if !UE_BUILD_SHIPPING
-					if (bSkipOptionalPakFiles == false || Filename.Find("optional") == INDEX_NONE)
-#endif
+					if (bSkipOptionalPakFiles && Filename.Contains("optional"))
+					{
+						return true;
+					}
+
+					if (!SkipPaksWildCard.IsEmpty() && Filename.MatchesWildcard(SkipPaksWildCard))
 					{
-						FoundPakFiles.Add(Filename);
+						return true;
 					}
+#endif
+
+					FoundPakFiles.Add(Filename);
 				}
 			}
 			return true;
@@ -5471,9 +5480,11 @@ void FPakPlatformFile::FindPakFilesInDirectory(IPlatformFile* LowLevelFile, cons
 	};
 
 	bool bSkipOptionalPakFiles = FParse::Param(FCommandLine::Get(), TEXT("SkipOptionalPakFiles"));
+	FString SkipPaksWildCard;
+	FParse::Value(FCommandLine::Get(), TEXT("-SkipStartupPaksWildCard="), SkipPaksWildCard);
 
 	// Find all pak files.
-	FPakSearchVisitor Visitor(OutPakFiles, WildCard, bSkipOptionalPakFiles);
+	FPakSearchVisitor Visitor(OutPakFiles, WildCard, SkipPaksWildCard, bSkipOptionalPakFiles);
 	LowLevelFile->IterateDirectoryRecursively(Directory, Visitor);
 }
 
@@ -5654,7 +5665,7 @@ bool FPakPlatformFile::Initialize(IPlatformFile* Inner, const TCHAR* CmdLine)
 	// Find and mount pak files from the specified directories.
 	TArray<FString> PakFolders;
 	GetPakFolders(FCommandLine::Get(), PakFolders);
-	MountAllPakFiles(PakFolders, *StartupPaksWildcard);
+	MountAllPakFiles(PakFolders, StartupPaksWildcard);
 
 #if !UE_BUILD_SHIPPING
 	GPakExec = MakeUnique<FPakExec>(*this);
