PakFiles: To reduce memory used by PakFiles, remove the DirectoryIndex TMap after load. Replace its functionality with an expansion of the DirectoryTree that we were previously using only for the directory paths. The DirectoryTree now contains the file paths as well.
#rnx

diff --git a/Engine/Source/Developer/PakFileUtilities/Private/PakFileUtilities.cpp b/Engine/Source/Developer/PakFileUtilities/Private/PakFileUtilities.cpp
--- a/Engine/Source/Developer/PakFileUtilities/Private/PakFileUtilities.cpp
+++ b/Engine/Source/Developer/PakFileUtilities/Private/PakFileUtilities.cpp
@@ -3204,7 +3204,7 @@ void WritePakFooter(FArchive& PakHandle, FPakFooterInfo& Footer)
 			{
 				FMemoryWriter& SecondaryWriter = PathHashIndexWriter.GetSecondaryWriter();
 				SecondaryWriter << PathHashIndex;
-				SecondaryWriter << PrunedDirectoryIndex;
+				FPakFile::SaveIndexInternal_DirectoryIndex(SecondaryWriter, PrunedDirectoryIndex);
 			}
 			PathHashIndexWriter.FinalizeAndRecordOffset(Footer.Info.IndexOffset + PrimaryIndexData.Num(), FinalizeIndexBlock);
 		}
@@ -3214,7 +3214,7 @@ void WritePakFooter(FArchive& PakHandle, FPakFooterInfo& Footer)
 		{
 			{
 				FMemoryWriter& SecondaryWriter = FullDirectoryIndexWriter.GetSecondaryWriter();
-				SecondaryWriter << DirectoryIndex;
+				FPakFile::SaveIndexInternal_DirectoryIndex(SecondaryWriter, DirectoryIndex);
 			}
 			FullDirectoryIndexWriter.FinalizeAndRecordOffset(Footer.Info.IndexOffset + PrimaryIndexData.Num() + PathHashIndexData.Num(), FinalizeIndexBlock);
 		}
@@ -3279,18 +3279,20 @@ bool ListFilesInPak(const TCHAR * InPakFilename, int64 SizeFilter, bool bInclude
 		UE_LOG(LogPakFile, Display, TEXT("Listing %s with mount point \"%s\""),
 			*FPaths::GetCleanFilename(InPakFilename), *PakFile.GetMountPoint());
 
-		TArray<FPakFile::FPakEntryIterator> Records;
+		TArray<TPair<TOptional<FString>, FPakEntry>> Records;
 
 		for (FPakFile::FPakEntryIterator It(PakFile,bIncludeDeleted); It; ++It)
 		{
-			Records.Add(It);
+			const FString* Filename = It.TryGetFilename();
+			Records.Add({ Filename ? TOptional<FString>(*Filename) : TOptional<FString>(), It.Info()});
 		}
 
 		struct FOffsetSort
 		{
-			FORCEINLINE bool operator()(const FPakFile::FPakEntryIterator& A, const FPakFile::FPakEntryIterator& B) const
+			FORCEINLINE bool operator()(const TPair<TOptional<FString>, FPakEntry>& A,
+				const TPair<TOptional<FString>, FPakEntry>& B) const
 			{
-				return A.Info().Offset < B.Info().Offset;
+				return A.Value.Offset < B.Value.Offset;
 			}
 		};
 
@@ -3305,9 +3307,9 @@ bool ListFilesInPak(const TCHAR * InPakFilename, int64 SizeFilter, bool bInclude
 		HashesBuffer.SetNum(FileCount * sizeof(FPakEntry::Hash));
 		uint8* Hashes = HashesBuffer.GetData();
 		int32 EntryIndex = 0;
-		for (auto It : Records)
+		for (const TPair<TOptional<FString>, FPakEntry>& It : Records)
 		{
-			PakFile.ReadHashFromPayload(It.Info(), Hashes + (EntryIndex++)*sizeof(FPakEntry::Hash));
+			PakFile.ReadHashFromPayload(It.Value, Hashes + (EntryIndex++)*sizeof(FPakEntry::Hash));
 		}
 
 		if (CSVFilename.Len() > 0)
@@ -3326,15 +3328,15 @@ bool ListFilesInPak(const TCHAR * InPakFilename, int64 SizeFilter, bool bInclude
 				OutputArchive->Logf(TEXT("Filename, Offset, Size, Hash, Deleted, Compressed, CompressionMethod"));
 			}
 			EntryIndex = 0;
-			for (auto It : Records)
+			for (const TPair<TOptional<FString>, FPakEntry>& It : Records)
 			{
-				const FPakEntry& Entry = It.Info();
+				const FPakEntry& Entry = It.Value;
 				uint8* Hash = Hashes + (EntryIndex++)*sizeof(FPakEntry::Hash);
 
 				bool bWasCompressed = Entry.CompressionMethodIndex != 0;
 				OutputArchive->Logf(
 					TEXT("%s%s, %lld, %lld, %s, %s, %s, %d"),
-					*MountPoint, It.TryGetFilename() ? **It.TryGetFilename() : TEXT("<FileNamesNotLoaded>"),
+					*MountPoint, It.Key ? **It.Key : TEXT("<FileNamesNotLoaded>"),
 					Entry.Offset, Entry.Size,
 					*BytesToHex(Hash, sizeof(FPakEntry::Hash)),
 					Entry.IsDeleteRecord() ? TEXT("true") : TEXT("false"),
@@ -3388,9 +3390,9 @@ bool ListFilesInPak(const TCHAR * InPakFilename, int64 SizeFilter, bool bInclude
 			}
 		}
 		EntryIndex = 0;
-		for (auto It : Records)
+		for (const TPair<TOptional<FString>, FPakEntry>& It : Records)
 		{
-			const FPakEntry& Entry = It.Info();
+			const FPakEntry& Entry = It.Value;
 			uint8* Hash = Hashes + (EntryIndex++) * sizeof(FPakEntry::Hash);
 			if (Entry.Size >= SizeFilter)
 			{
@@ -3403,8 +3405,10 @@ bool ListFilesInPak(const TCHAR * InPakFilename, int64 SizeFilter, bool bInclude
 					{
 						if (InspectChunks.Contains(Chunk))
 						{
-							UE_LOG(LogPakFile, Display, TEXT("[%d - %d] \"%s%s\" offset: %lld, size: %d bytes, sha1: %s, compression: %s."), FirstChunk, LastChunk, *MountPoint,
-								It.TryGetFilename() ? **It.TryGetFilename() : TEXT("<FileNamesNotLoaded"), Entry.Offset, Entry.Size, *BytesToHex(Hash, sizeof(FPakEntry::Hash)),
+							UE_LOG(LogPakFile, Display, TEXT("[%d - %d] \"%s%s\" offset: %lld, size: %d bytes, sha1: %s, compression: %s."),
+								FirstChunk, LastChunk, *MountPoint,
+								It.Key ? **It.Key : TEXT("<FileNamesNotLoaded"), Entry.Offset, Entry.Size,
+								*BytesToHex(Hash, sizeof(FPakEntry::Hash)),
 								*PakFile.GetInfo().GetCompressionMethod(Entry.CompressionMethodIndex).ToString());
 							break;
 						}
@@ -3412,8 +3416,10 @@ bool ListFilesInPak(const TCHAR * InPakFilename, int64 SizeFilter, bool bInclude
 				}
 				else
 				{
-					UE_LOG(LogPakFile, Display, TEXT("\"%s%s\" offset: %lld, size: %d bytes, sha1: %s, compression: %s."), *MountPoint, It.TryGetFilename() ? **It.TryGetFilename() : TEXT("<FileNamesNotLoaded"),
-						Entry.Offset, Entry.Size, *BytesToHex(Hash, sizeof(FPakEntry::Hash)), *PakFile.GetInfo().GetCompressionMethod(Entry.CompressionMethodIndex).ToString());
+					UE_LOG(LogPakFile, Display, TEXT("\"%s%s\" offset: %lld, size: %d bytes, sha1: %s, compression: %s."),
+						*MountPoint, It.Key ? **It.Key : TEXT("<FileNamesNotLoaded"),
+						Entry.Offset, Entry.Size, *BytesToHex(Hash, sizeof(FPakEntry::Hash)),
+						*PakFile.GetInfo().GetCompressionMethod(Entry.CompressionMethodIndex).ToString());
 				}
 			}
 			FileSize += Entry.Size;
diff --git a/Engine/Source/Runtime/PakFile/Private/IPlatformFilePak.cpp b/Engine/Source/Runtime/PakFile/Private/IPlatformFilePak.cpp
--- a/Engine/Source/Runtime/PakFile/Private/IPlatformFilePak.cpp
+++ b/Engine/Source/Runtime/PakFile/Private/IPlatformFilePak.cpp
@@ -217,14 +217,14 @@ static FAutoConsoleVariableRef CVar_ValidateDirectoryTreeSearchConsistency(
 	TEXT("pak.ValidateDirectoryTreeSearchConsistency"),
 	GPak_ValidateDirectoryTreeSearchConsistency,
 	TEXT("Validate consistency of directory tree index search by comparing results with directory index search"));
-#endif
+#endif // !UE_BUILD_SHIPPING
 
 bool GPak_UseDirectoryTreeForPakSearch = true;
 static FAutoConsoleVariableRef CVar_UseDirectoryTreeForPakSearch(
 	TEXT("pak.UseDirectoryTreeForPakSearch"),
 	GPak_UseDirectoryTreeForPakSearch,
 	TEXT("Enable use of directory tree when finding files"));
-#endif
+#endif // ENABLE_PAKFILE_USE_DIRECTORY_TREE
 
 
 TPakChunkHash ComputePakChunkHash(const void* InData, int64 InDataSizeInBytes)
@@ -5733,7 +5733,18 @@ uint64 GetRecursiveAllocatedSize(const FPakFile::FDirectoryIndex& Index)
 	}
 	return Size;
 }
-#endif
+#if ENABLE_PAKFILE_USE_DIRECTORY_TREE
+uint64 GetRecursiveAllocatedSize(const UE::PakFile::Private::FDirectoryTreeIndex& Index)
+{
+	uint64 Size = Index.GetAllocatedSize();
+	for (const TPair<FStringView, const FPakDirectory&> kvpair : Index)
+	{
+		Size += GetRecursiveAllocatedSize(kvpair.Value);
+	}
+	return Size;
+}
+#endif // ENABLE_PAKFILE_USE_DIRECTORY_TREE
+#endif // !UE_BUILD_SHIPPING
 
 static float GPakReaderReleaseDelay = 5.0f;
 static FAutoConsoleVariableRef CVarPakReaderReleaseDelay(
@@ -5791,7 +5802,7 @@ void FPakPlatformFile::OptimizeMemoryUsageForMountedPaks()
 #if ENABLE_PAKFILE_RUNTIME_PRUNING
 	if (FPakFile::bSomePakNeedsPruning)
 	{
-		for (auto& Pak : Paks)
+		for (FPakListEntry& Pak : Paks)
 		{
 			FPakFile* PakFile = Pak.PakFile;
 			if (PakFile->bWillPruneDirectoryIndex)
@@ -5800,28 +5811,23 @@ void FPakPlatformFile::OptimizeMemoryUsageForMountedPaks()
 				FWriteScopeLock DirectoryLock(PakFile->DirectoryIndexLock);
 				if (PakFile->bNeedsLegacyPruning)
 				{
-					FPakFile::PruneDirectoryIndex(PakFile->DirectoryIndex, &PakFile->PrunedDirectoryIndex, PakFile->MountPoint);
+					FPakFile::PruneDirectoryIndexInternal(&PakFile->DirectoryIndex, &PakFile->DirectoryTreeIndex,
+						&PakFile->PrunedDirectoryIndex, &PakFile->PrunedDirectoryTreeIndex, PakFile->MountPoint);
 					PakFile->bNeedsLegacyPruning = false;
 				}
 
 				Swap(PakFile->DirectoryIndex, PakFile->PrunedDirectoryIndex);
 				PakFile->PrunedDirectoryIndex.Empty();
+#if ENABLE_PAKFILE_USE_DIRECTORY_TREE
+				Swap(PakFile->DirectoryTreeIndex, PakFile->PrunedDirectoryTreeIndex);
+				PakFile->PrunedDirectoryTreeIndex.Empty();
+#endif // ENABLE_PAKFILE_USE_DIRECTORY_TREE
 				PakFile->bHasFullDirectoryIndex = false;
 				PakFile->bWillPruneDirectoryIndex = false;
-#if ENABLE_PAKFILE_USE_DIRECTORY_TREE
-				{
-					SCOPED_BOOT_TIMING("PakFile_ReinitializePruneDirectoryTree");
-					PakFile->DirectoryTreeIndex.Empty();
-					for (TMap<FString, FPakDirectory>::TIterator It(PakFile->DirectoryIndex); It; ++It)
-					{
-						PakFile->DirectoryTreeIndex.FindOrAdd(It.Key()) = &It.Value();
-					}
-				}
-#endif
 			}
 		}
 	}
-#endif
+#endif // ENABLE_PAKFILE_RUNTIME_PRUNING
 
 #if !UE_BUILD_SHIPPING
 	uint64 DirectoryHashSize = 0;
@@ -5829,26 +5835,29 @@ void FPakPlatformFile::OptimizeMemoryUsageForMountedPaks()
 	uint64 EntriesSize = 0;
 	uint64 DirTreeIndexSize = 0;
 
-	for (auto& Pak : Paks)
+	for (FPakListEntry& Pak : Paks)
 	{
 		FPakFile* PakFile = Pak.PakFile;
 		{
 			FPakFile::FScopedPakDirectoryIndexAccess ScopeAccess(*PakFile);
 			DirectoryHashSize += GetRecursiveAllocatedSize(PakFile->DirectoryIndex);
-#if ENABLE_PAKFILE_RUNTIME_PRUNING 
-			{
-				DirectoryHashSize += GetRecursiveAllocatedSize(PakFile->PrunedDirectoryIndex);
-			}
-		}
+#if ENABLE_PAKFILE_USE_DIRECTORY_TREE
+			DirTreeIndexSize += GetRecursiveAllocatedSize(PakFile->DirectoryTreeIndex);
 #endif
+
+#if ENABLE_PAKFILE_RUNTIME_PRUNING
+			DirectoryHashSize += GetRecursiveAllocatedSize(PakFile->PrunedDirectoryIndex);
+#if ENABLE_PAKFILE_USE_DIRECTORY_TREE
+			DirTreeIndexSize += GetRecursiveAllocatedSize(PakFile->PrunedDirectoryTreeIndex);
+#endif // ENABLE_PAKFILE_USE_DIRECTORY_TREE
+#endif // ENABLE_PAKFILE_RUNTIME_PRUNING
+		}
 		PathHashSize += PakFile->PathHashIndex.GetAllocatedSize();
 		EntriesSize += PakFile->EncodedPakEntries.GetAllocatedSize();
 		EntriesSize += PakFile->Files.GetAllocatedSize();
-#if ENABLE_PAKFILE_USE_DIRECTORY_TREE
-		DirTreeIndexSize += PakFile->DirectoryTreeIndex.GetAllocatedSize();
-#endif
 	}
-	UE_LOG(LogPakFile, Log, TEXT("AllPaks IndexSizes: DirectoryHashSize=%llu, PathHashSize=%llu, EntriesSize=%llu, DirTreeIndexSize=%llu, TotalSize=%llu"), DirectoryHashSize, PathHashSize, EntriesSize, DirTreeIndexSize, DirectoryHashSize + PathHashSize + EntriesSize + DirTreeIndexSize);
+	UE_LOG(LogPakFile, Log, TEXT("AllPaks IndexSizes: DirectoryHashSize=%llu, PathHashSize=%llu, EntriesSize=%llu, DirTreeIndexSize=%llu, TotalSize=%llu"),
+	 DirectoryHashSize, PathHashSize, EntriesSize, DirTreeIndexSize, DirectoryHashSize + PathHashSize + EntriesSize + DirTreeIndexSize);
 #endif
 }
 
@@ -6561,6 +6570,9 @@ void FPakPlatformFile::MakeUniquePakFilesForTheseFiles(const TArray<TArray<FStri
 
 				if (NewPakFile != nullptr)
 				{
+#if ENABLE_PAKFILE_USE_DIRECTORY_TREE
+					NewPakFile->bUseDirectoryTree = GPak_UseDirectoryTreeForPakSearch;
+#endif
 //					NewPakFile->SetUnderlyingCacheTrimDisabled(true);
 					NewPakFile->AddSpecialFile(FileEntry, *InFiles[k][i]);
 				}
diff --git a/Engine/Source/Runtime/PakFile/Private/PakFile.cpp b/Engine/Source/Runtime/PakFile/Private/PakFile.cpp
--- a/Engine/Source/Runtime/PakFile/Private/PakFile.cpp
+++ b/Engine/Source/Runtime/PakFile/Private/PakFile.cpp
@@ -116,6 +116,9 @@ FPakFile::FPakFile(const TCHAR* Filename, bool bIsSigned)
 #if ENABLE_PAKFILE_RUNTIME_PRUNING
 	, bWillPruneDirectoryIndex(false)
 	, bNeedsLegacyPruning(false)
+#endif
+#if ENABLE_PAKFILE_USE_DIRECTORY_TREE
+	, bUseDirectoryTree(false)
 #endif
 	, PakchunkIndex(GetPakchunkIndexFromPakFile(Filename))
  	, MappedFileHandle(nullptr)
@@ -146,6 +149,9 @@ FPakFile::FPakFile(IPlatformFile* LowerLevel, const TCHAR* Filename, bool bIsSig
 #if ENABLE_PAKFILE_RUNTIME_PRUNING
 	, bWillPruneDirectoryIndex(false)
 	, bNeedsLegacyPruning(false)
+#endif
+#if ENABLE_PAKFILE_USE_DIRECTORY_TREE
+	, bUseDirectoryTree(false)
 #endif
 	, PakchunkIndex(GetPakchunkIndexFromPakFile(Filename))
 	, MappedFileHandle(nullptr)
@@ -173,6 +179,9 @@ FPakFile::FPakFile(FArchive* Archive)
 #if ENABLE_PAKFILE_RUNTIME_PRUNING
 	, bWillPruneDirectoryIndex(false)
 	, bNeedsLegacyPruning(false)
+#endif
+#if ENABLE_PAKFILE_USE_DIRECTORY_TREE
+	, bUseDirectoryTree(false)
 #endif
 	, PakchunkIndex(INDEX_NONE)
 	, MappedFileHandle(nullptr)
@@ -337,12 +346,26 @@ void FPakFile::Initialize(FArchive& Reader, bool bLoadIndex)
 
 void FPakFile::LoadIndex(FArchive& Reader)
 {
+#if ENABLE_PAKFILE_USE_DIRECTORY_TREE
+	bUseDirectoryTree = GPak_UseDirectoryTreeForPakSearch;
+	bool bStoreDirectoryTree = bUseDirectoryTree;
+	bool bStoreDirectoryTMap = !bUseDirectoryTree;
+#if !UE_BUILD_SHIPPING
+	bStoreDirectoryTMap |= GPak_ValidateDirectoryTreeSearchConsistency;
+#endif // !UE_BUILD_SHIPPING
+#else // !ENABLE_PAKFILE_USE_DIRECTORY_TREE
+	constexpr bool bStoreDirectoryTMap = true;
+	constexpr bool bStoreDirectoryTree = false;
+#endif // else !ENABLE_PAKFILE_USE_DIRECTORY_TREE
+
 	if (Info.Version >= FPakInfo::PakFile_Version_PathHashIndex)
 	{
-		if (!LoadIndexInternal(Reader))
+		if (!LoadIndexInternal(Reader, DirectoryIndex, DirectoryTreeIndex, PrunedDirectoryIndex, PrunedDirectoryTreeIndex,
+			bStoreDirectoryTMap, bStoreDirectoryTree))
 		{
 			// Index loading failed. Try again
-			if (!LoadIndexInternal(Reader))
+			if (!LoadIndexInternal(Reader, DirectoryIndex, DirectoryTreeIndex, PrunedDirectoryIndex, PrunedDirectoryTreeIndex,
+				bStoreDirectoryTMap, bStoreDirectoryTree))
 			{
 				UE_LOG(LogPakFile, Fatal, TEXT("Corrupt pak index detected on pak file: %s"), *PakFilename);
 			}
@@ -351,27 +374,36 @@ void FPakFile::LoadIndex(FArchive& Reader)
 	else
 	{
 		SCOPED_BOOT_TIMING("PakFile_LoadLegacy");
-		if (!LoadLegacyIndex(Reader))
+		FDirectoryIndex LoadedTMap;
+		if (!LoadLegacyIndex(Reader, LoadedTMap))
 		{
 			// Index loading failed. Try again
-			if (!LoadLegacyIndex(Reader))
+			if (!LoadLegacyIndex(Reader, LoadedTMap))
 			{
 				UE_LOG(LogPakFile, Fatal, TEXT("Corrupt pak index detected on pak file: %s"), *PakFilename);
 			}
 		}
-	}
 #if ENABLE_PAKFILE_USE_DIRECTORY_TREE
-	{
-		SCOPED_BOOT_TIMING("PakFile_PopulateDirectoryTree");
-		for (TMap<FString, FPakDirectory>::TIterator It(DirectoryIndex); It; ++It)
+		if (bStoreDirectoryTree)
 		{
-			DirectoryTreeIndex.FindOrAdd(It.Key()) = &It.Value();
+			for (TMap<FString, FPakDirectory>::TIterator It(LoadedTMap); It; ++It)
+			{
+				DirectoryTreeIndex.FindOrAdd(It->Key) =
+					bStoreDirectoryTMap ? FPakDirectory(It->Value) : MoveTemp(It->Value);
+			}
+			DirectoryTreeIndex.Shrink();
+		}
+#endif
+		if (bStoreDirectoryTMap)
+		{
+			DirectoryIndex = MoveTemp(LoadedTMap);
 		}
 	}
-#endif	// ENABLE_PAKFILE_USE_DIRECTORY_TREE
 }
 
-bool FPakFile::LoadIndexInternal(FArchive& Reader)
+bool FPakFile::LoadIndexInternal(FArchive& Reader, FDirectoryIndex& OutDirectoryTMap,
+	FDirectoryTreeIndex& OutDirectoryTree, FDirectoryIndex& OutPrunedDirectoryTMap,
+	FDirectoryTreeIndex& OutPrunedDirectoryTree, bool bStoreDirectoryTMap, bool bStoreDirectoryTree)
 {
 	bHasPathHashIndex = false;
 	bHasFullDirectoryIndex = false;
@@ -580,7 +612,8 @@ bool FPakFile::LoadIndexInternal(FArchive& Reader)
 		// Store the PrunedDirectoryIndex in our DirectoryIndex
 		{
 			SCOPED_BOOT_TIMING("PakFile_SerializePrunedDirectoryIndex");
-			PathHashIndexReader << DirectoryIndex;
+			LoadIndexInternal_DirectoryIndex(PathHashIndexReader, OutDirectoryTMap, OutDirectoryTree,
+				bStoreDirectoryTMap, bStoreDirectoryTree);
 		}
 		bHasFullDirectoryIndex = false;
 #if ENABLE_PAKFILE_RUNTIME_PRUNING
@@ -630,7 +663,8 @@ bool FPakFile::LoadIndexInternal(FArchive& Reader)
 		FMemoryReader SecondaryIndexReader(FullDirectoryIndexData);
 		{
 			SCOPED_BOOT_TIMING("PakFile_SerializeDirectoryIndex");
-			SecondaryIndexReader << DirectoryIndex;
+			LoadIndexInternal_DirectoryIndex(SecondaryIndexReader, OutDirectoryTMap, OutDirectoryTree,
+				bStoreDirectoryTMap, bStoreDirectoryTree);
 		}
 		bHasFullDirectoryIndex = true;
 
@@ -645,7 +679,8 @@ bool FPakFile::LoadIndexInternal(FArchive& Reader)
 			check(bWillUsePathHashIndex); // Need to confirm that we have read the PathHashIndex bytes
 			{
 				SCOPED_BOOT_TIMING("PakFile_SerializePrunedDirectoryIndex");
-				PathHashIndexReader << PrunedDirectoryIndex;
+				LoadIndexInternal_DirectoryIndex(PathHashIndexReader, OutPrunedDirectoryTMap, OutPrunedDirectoryTree,
+					bStoreDirectoryTMap, bStoreDirectoryTree);
 			}
 			bWillPruneDirectoryIndex = true;
 			bSomePakNeedsPruning = true;
@@ -661,7 +696,76 @@ bool FPakFile::LoadIndexInternal(FArchive& Reader)
 	return true;
 }
 
-bool FPakFile::LoadLegacyIndex(FArchive& Reader)
+void FPakFile::LoadIndexInternal_DirectoryIndex(FArchive& Ar, FDirectoryIndex& OutDirectoryTMap,
+	FDirectoryTreeIndex& OutDirectoryTree, bool bLoadIntoDirectoryTMap, bool bLoadIntoDirectoryTree)
+{
+	FString DirectoryName;
+	FString FileName;
+	FPakEntryLocation FileData;
+	FPakDirectory DirectoryData;
+
+	int32 NumDirectories = 0;
+	Ar << NumDirectories;
+	OutDirectoryTMap.Reset();
+	if (bLoadIntoDirectoryTMap)
+	{
+		OutDirectoryTMap.Reserve(NumDirectories);
+	}
+#if ENABLE_PAKFILE_USE_DIRECTORY_TREE
+	OutDirectoryTree.Empty();
+#endif
+	for (; NumDirectories > 0; --NumDirectories)
+	{
+		Ar << DirectoryName;
+		int32 NumFiles = 0;
+		Ar << NumFiles;
+
+		DirectoryData.Reset();
+		DirectoryData.Reserve(NumFiles);
+		for (; NumFiles > 0; --NumFiles)
+		{
+			Ar << FileName;
+			Ar << FileData;
+			DirectoryData.Add(FileName, FileData);
+		}
+
+#if ENABLE_PAKFILE_USE_DIRECTORY_TREE
+		if (bLoadIntoDirectoryTree)
+		{
+			OutDirectoryTree.FindOrAdd(DirectoryName) =
+				bLoadIntoDirectoryTMap ? FPakDirectory(DirectoryData) : MoveTemp(DirectoryData);
+
+		}
+#endif
+		if (bLoadIntoDirectoryTMap)
+		{
+			OutDirectoryTMap.FindOrAdd(DirectoryName) = MoveTemp(DirectoryData);
+		}
+	}
+#if ENABLE_PAKFILE_USE_DIRECTORY_TREE
+	OutDirectoryTree.Shrink();
+#endif
+}
+
+void FPakFile::SaveIndexInternal_DirectoryIndex(FArchive& Ar, const FDirectoryIndex& DirectoryTMap)
+{
+	int32 NumDirectories = DirectoryTMap.Num();
+	Ar << NumDirectories;
+	for (const TPair<FString, FPakDirectory>& DirPair : DirectoryTMap)
+	{
+		Ar << const_cast<FString&>(DirPair.Key);
+		int32 NumFiles = DirPair.Value.Num();
+		Ar << NumFiles;
+
+		for (const TPair<FString, FPakEntryLocation>& FilePair : DirPair.Value)
+		{
+			Ar << const_cast<FString&>(FilePair.Key);
+			Ar << const_cast<FPakEntryLocation&>(FilePair.Value);
+		}
+	}
+}
+
+bool FPakFile::LoadLegacyIndex(FArchive& Reader, FDirectoryIndex& OutDirectoryTMap)
 {
 	bHasPathHashIndex = false;
 	bHasFullDirectoryIndex = false;
@@ -746,7 +850,7 @@ bool FPakFile::LoadLegacyIndex(FArchive& Reader)
 #endif
 	FPathHashIndex* PathHashToWrite = bCreatePathHash ? &PathHashIndex : nullptr;
 	FPakFile::EncodePakEntriesIntoIndex(NumEntries, ReadNextEntry, *PakFilename, Info, MountPoint, NumEncodedEntries, NumDeletedEntries, &PathHashSeed,
-		&DirectoryIndex, PathHashToWrite, EncodedPakEntries, Files, &CollisionDetection, Info.Version);
+		&OutDirectoryTMap, PathHashToWrite, EncodedPakEntries, Files, &CollisionDetection, Info.Version);
 	check(NumEncodedEntries + Files.Num() + NumDeletedEntries == NumEntries);
 	Files.Shrink();
 	EncodedPakEntries.Shrink();
@@ -851,11 +955,19 @@ void FPakFile::EncodePakEntriesIntoIndex(int32 InNumEntries, const ReadNextEntry
 		}
 
 		// Add the Entry into the requested Indexes
-		AddEntryToIndex(Pair.Filename, EntryLocation, MountPoint, PathHashSeed, OutDirectoryIndex, OutPathHashIndex, InOutCollisionDetection, PakFileVersion);
+		AddEntryToIndex(Pair.Filename, EntryLocation, MountPoint, PathHashSeed, OutDirectoryIndex,
+			nullptr /* DirectoryTreeIndex */, OutPathHashIndex, InOutCollisionDetection, PakFileVersion);
 	}
 }
 
-void FPakFile::PruneDirectoryIndex(FDirectoryIndex& InOutDirectoryIndex, FDirectoryIndex* PrunedDirectoryIndex, const FString& MountPoint)
+void FPakFile::PruneDirectoryIndex(FDirectoryIndex& InOutDirectoryIndex, FDirectoryIndex* InPrunedDirectoryIndex,
+	const FString& MountPoint)
+{
+	PruneDirectoryIndexInternal(&InOutDirectoryIndex, nullptr, InPrunedDirectoryIndex, nullptr, MountPoint);
+}
+
+void FPakFile::PruneDirectoryIndexInternal(FDirectoryIndex* InOutDirectoryIndex, FDirectoryTreeIndex* InOutDirectoryTreeIndex,
+	FDirectoryIndex* InPrunedDirectoryIndex, FDirectoryTreeIndex* InPrunedDirectoryTreeIndex, const FString& MountPoint)
 {
 	// Caller holds WriteLock on DirectoryIndexLock
 	TArray<FString> FileWildCards, DirectoryWildCards, OldWildCards;
@@ -866,133 +978,273 @@ void FPakFile::PruneDirectoryIndex(FDirectoryIndex& InOutDirectoryIndex, FDirect
 	FileWildCards.Append(OldWildCards);
 	int32 NumKeptEntries = 0;
 
-	if (PrunedDirectoryIndex)
+	if (InPrunedDirectoryIndex)
+	{
+		InPrunedDirectoryIndex->Empty();
+	}
+#if ENABLE_PAKFILE_USE_DIRECTORY_TREE
+	if (InPrunedDirectoryTreeIndex)
 	{
-		PrunedDirectoryIndex->Empty();
+		InPrunedDirectoryTreeIndex->Empty();
 	}
+#endif
 
 	TMap<FString, bool> KeepDirectory;
 
 	// Clear out those portions of the Index allowed by the user.
 	if (DirectoryWildCards.Num() > 0 || FileWildCards.Num() > 0)
 	{
-		for (auto It = InOutDirectoryIndex.CreateIterator(); It; ++It)
+		if (InOutDirectoryIndex)
 		{
-			const FString& DirectoryName = It.Key();
-			FPakDirectory& OriginalDirectory = It.Value();
-			const FString FullDirectoryName = PakPathCombine(MountPoint, DirectoryName);
-			check(IsPathInDirectoryFormat(FullDirectoryName));
-			FPakDirectory* PrunedDirectory = nullptr;
-			bool bKeepDirectory = false;
-
-			TArray<FString> RemoveFilenames;
-			for (auto FileIt = It->Value.CreateIterator(); FileIt; ++FileIt)
+			for (FDirectoryIndex::TIterator It = InOutDirectoryIndex->CreateIterator(); It; ++It)
 			{
-				const FString& FileNameWithoutPath = FileIt->Key;
-				const FString FullFilename = PakPathCombine(FullDirectoryName, FileNameWithoutPath);
-				bool bKeepFile = false;
-
-				for (const FString& WildCard : FileWildCards)
+				const FString& DirectoryName = It.Key();
+				FPakDirectory& OriginalDirectory = It.Value();
+				const FString FullDirectoryName = PakPathCombine(MountPoint, DirectoryName);
+				check(IsPathInDirectoryFormat(FullDirectoryName));
+				FPakDirectory* PrunedDirectory = nullptr;
+				bool bKeepDirectory = false;
+
+				TArray<FString> RemoveFilenames;
+				for (auto FileIt = It->Value.CreateIterator(); FileIt; ++FileIt)
 				{
-					if (FullFilename.MatchesWildcard(WildCard))
+					const FString& FileNameWithoutPath = FileIt->Key;
+					const FString FullFilename = PakPathCombine(FullDirectoryName, FileNameWithoutPath);
+					bool bKeepFile = false;
+
+					for (const FString& WildCard : FileWildCards)
 					{
-						bKeepFile = true;
-						break;
+						if (FullFilename.MatchesWildcard(WildCard))
+						{
+							bKeepFile = true;
+							break;
+						}
 					}
-				}
 
-				if (bKeepFile)
-				{
-					bKeepDirectory = true;
-					if (PrunedDirectoryIndex)
+					if (bKeepFile)
 					{
-						if (!PrunedDirectory)
+						bKeepDirectory = true;
+						if (InPrunedDirectoryIndex)
 						{
-							PrunedDirectory = &PrunedDirectoryIndex->Add(DirectoryName);
+							if (!PrunedDirectory)
+							{
+								PrunedDirectory = &InPrunedDirectoryIndex->Add(DirectoryName);
+							}
+							PrunedDirectory->Add(FileNameWithoutPath, FileIt->Value);
+						}
+					}
+					else
+					{
+						if (!InPrunedDirectoryIndex)
+						{
+							RemoveFilenames.Add(FileNameWithoutPath);
 						}
-						PrunedDirectory->Add(FileNameWithoutPath, FileIt->Value);
 					}
 				}
-				else
+				if (!InPrunedDirectoryIndex)
+				{
+					for (const FString& FileNameWithoutPath : RemoveFilenames)
+					{
+						OriginalDirectory.Remove(FileNameWithoutPath);
+					}
+				}
+
+				if (!bKeepDirectory)
 				{
-					if (!PrunedDirectoryIndex)
+					for (const FString& WildCard : DirectoryWildCards)
 					{
-						RemoveFilenames.Add(FileNameWithoutPath);
+						if (FullDirectoryName.MatchesWildcard(WildCard))
+						{
+							bKeepDirectory = true;
+							break;
+						}
 					}
 				}
+				KeepDirectory.FindOrAdd(DirectoryName) = bKeepDirectory;
 			}
-			if (!PrunedDirectoryIndex)
+
+			// For each kept directory, mark that we need to keep all of its parents up to the mount point
+			for (const TPair<FString, bool>& Pair : KeepDirectory)
 			{
-				for (const FString& FileNameWithoutPath : RemoveFilenames)
+				if (Pair.Value)
 				{
-					OriginalDirectory.Remove(FileNameWithoutPath);
+					FString CurrentDirectory = Pair.Key;
+					FString UnusedCleanFileName;
+					while (CurrentDirectory != MountPoint)
+					{
+						if (!SplitPathInline(CurrentDirectory, UnusedCleanFileName))
+						{
+							break;
+						}
+						bool& bOldValue = KeepDirectory.FindOrAdd(CurrentDirectory);
+						if (bOldValue)
+						{
+							break;
+						}
+						bOldValue = true;
+					}
 				}
 			}
 
-			if (!bKeepDirectory)
+			// Prune all of the directories
+			for (const TPair<FString, bool>& Pair : KeepDirectory)
 			{
-				for (const FString& WildCard : DirectoryWildCards)
+				const FString& DirectoryName = Pair.Key;
+				bool bKeep = Pair.Value;
+				if (bKeep)
 				{
-					if (FullDirectoryName.MatchesWildcard(WildCard))
+					if (InPrunedDirectoryIndex)
 					{
-						bKeepDirectory = true;
-						break;
+						InPrunedDirectoryIndex->FindOrAdd(DirectoryName);
+					}
+				}
+				else
+				{
+					if (!InPrunedDirectoryIndex)
+					{
+						InOutDirectoryIndex->Remove(DirectoryName);
 					}
 				}
 			}
-			KeepDirectory.FindOrAdd(DirectoryName) = bKeepDirectory;
 		}
-
-		// For each kept directory, mark that we need to keep all of its parents up to the mount point
-		for (const TPair<FString, bool>& Pair : KeepDirectory)
+#if ENABLE_PAKFILE_USE_DIRECTORY_TREE
+		if (InOutDirectoryTreeIndex)
 		{
-			if (Pair.Value)
+			for (FDirectoryTreeIndex::FIterator It(InOutDirectoryTreeIndex->CreateIterator()); It; ++It)
 			{
-				FString CurrentDirectory = Pair.Key;
-				FString UnusedCleanFileName;
-				while (CurrentDirectory != MountPoint)
+				FString DirectoryName = FString(It->Key);
+				FPakDirectory& OriginalDirectory = It->Value;
+				FString FullDirectoryName = PakPathCombine(MountPoint, DirectoryName);
+				if (!FullDirectoryName.EndsWith(TEXT("/")))
+				{
+					FullDirectoryName += '/';
+				}
+				check(IsPathInDirectoryFormat(FullDirectoryName));
+				FPakDirectory* PrunedDirectory = nullptr;
+				bool bKeepDirectory = false;
+
+				TArray<FString> RemoveFilenames;
+				for (FPakDirectory::TIterator FileIt = OriginalDirectory.CreateIterator(); FileIt; ++FileIt)
 				{
-					if (!SplitPathInline(CurrentDirectory, UnusedCleanFileName))
+					const FString& FileNameWithoutPath = FileIt->Key;
+					const FString FullFilename = PakPathCombine(FullDirectoryName, FileNameWithoutPath);
+					bool bKeepFile = false;
+
+					for (const FString& WildCard : FileWildCards)
+					{
+						if (FullFilename.MatchesWildcard(WildCard))
+						{
+							bKeepFile = true;
+							break;
+						}
+					}
+
+					if (bKeepFile)
 					{
-						break;
+						bKeepDirectory = true;
+						if (InPrunedDirectoryTreeIndex)
+						{
+							if (!PrunedDirectory)
+							{
+								PrunedDirectory = &InPrunedDirectoryTreeIndex->FindOrAdd(DirectoryName);
+							}
+							PrunedDirectory->Add(FileNameWithoutPath, FileIt->Value);
+						}
 					}
-					bool& bOldValue = KeepDirectory.FindOrAdd(CurrentDirectory);
-					if (bOldValue)
+					else
 					{
-						break;
+						if (!InPrunedDirectoryTreeIndex)
+						{
+							RemoveFilenames.Add(FileNameWithoutPath);
+						}
 					}
-					bOldValue = true;
 				}
+				if (!InPrunedDirectoryTreeIndex)
+				{
+					for (const FString& FileNameWithoutPath : RemoveFilenames)
+					{
+						OriginalDirectory.Remove(FileNameWithoutPath);
+					}
+				}
+
+				if (!bKeepDirectory)
+				{
+					for (const FString& WildCard : DirectoryWildCards)
+					{
+						if (FullDirectoryName.MatchesWildcard(WildCard))
+						{
+							bKeepDirectory = true;
+							break;
+						}
+					}
+				}
+				KeepDirectory.FindOrAdd(DirectoryName) = bKeepDirectory;
 			}
-		}
 
-		// Prune all of the directories
-		for (const TPair<FString, bool>& Pair : KeepDirectory)
-		{
-			const FString& DirectoryName = Pair.Key;
-			bool bKeep = Pair.Value;
-			if (bKeep)
+			// For each kept directory, mark that we need to keep all of its parents up to the mount point
+			for (const TPair<FString, bool>& Pair : KeepDirectory)
 			{
-				if (PrunedDirectoryIndex)
+				if (Pair.Value)
 				{
-					PrunedDirectoryIndex->FindOrAdd(DirectoryName);
+					FString CurrentDirectory = Pair.Key;
+					FString UnusedCleanFileName;
+					while (CurrentDirectory != MountPoint)
+					{
+						if (!SplitPathInline(CurrentDirectory, UnusedCleanFileName))
+						{
+							break;
+						}
+						bool& bOldValue = KeepDirectory.FindOrAdd(CurrentDirectory);
+						if (bOldValue)
+						{
+							break;
+						}
+						bOldValue = true;
+					}
 				}
 			}
-			else
+
+			// Prune all of the directories
+			for (const TPair<FString, bool>& Pair : KeepDirectory)
 			{
-				if (!PrunedDirectoryIndex)
+				const FString& DirectoryName = Pair.Key;
+				bool bKeep = Pair.Value;
+				if (bKeep)
 				{
-					InOutDirectoryIndex.Remove(DirectoryName);
+					if (InPrunedDirectoryTreeIndex)
+					{
+						InPrunedDirectoryTreeIndex->FindOrAdd(DirectoryName);
+					}
+				}
+				else
+				{
+					if (!InPrunedDirectoryTreeIndex)
+					{
+						InOutDirectoryIndex->Remove(DirectoryName);
+					}
 				}
 			}
 		}
+#endif
 	}
 	else
 	{
-		if (!PrunedDirectoryIndex)
+		if (InOutDirectoryIndex)
 		{
-			InOutDirectoryIndex.Empty();
+			if (!InPrunedDirectoryIndex)
+			{
+				InOutDirectoryIndex->Empty();
+			}
+		}
+#if ENABLE_PAKFILE_USE_DIRECTORY_TREE
+		if (InOutDirectoryTreeIndex)
+		{
+			if (!InPrunedDirectoryTreeIndex)
+			{
+				InOutDirectoryTreeIndex->Empty();
+			}
 		}
+#endif
 	}
 }
 
@@ -1106,8 +1358,19 @@ bool FPakFile::ShouldValidatePrunedDirectory() const
 #endif
 }
 
-void FPakFile::AddEntryToIndex(const FString& Filename, const FPakEntryLocation& EntryLocation, const FString& MountPoint, uint64 PathHashSeed,
-	FDirectoryIndex* DirectoryIndex, FPathHashIndex* PathHashIndex, TMap<uint64, FString>* CollisionDetection, int32 PakFileVersion)
+bool FPakFile::ShouldUseDirectoryTree() const
+{
+#if ENABLE_PAKFILE_USE_DIRECTORY_TREE
+	return bUseDirectoryTree;
+#else
+	return false;
+#endif
+}
+
+void FPakFile::AddEntryToIndex(const FString& Filename, const FPakEntryLocation& EntryLocation,
+	const FString& MountPoint, uint64 PathHashSeed, FDirectoryIndex* DirectoryIndex,
+	FDirectoryTreeIndex* DirectoryTreeIndex, FPathHashIndex* PathHashIndex, TMap<uint64, FString>* CollisionDetection,
+	int32 PakFileVersion)
 {
 	FString RelativePathFromMount;
 	if (FPaths::IsRelative(Filename))
@@ -1149,6 +1412,34 @@ void FPakFile::AddEntryToIndex(const FString& Filename, const FPakEntryLocation&
 		}
 		Directory->Add(CleanFileName, EntryLocation);
 	}
+#if ENABLE_PAKFILE_USE_DIRECTORY_TREE
+	if (DirectoryTreeIndex)
+	{
+		FString RelativeDirectoryFromMount = RelativePathFromMount;
+		FString CleanFileName;
+		SplitPathInline(RelativeDirectoryFromMount, CleanFileName);
+		FPakDirectory* Directory = DirectoryTreeIndex->Find(RelativeDirectoryFromMount);
+		if (Directory == nullptr)
+		{
+			// add the parent directories up to the mount point (mount point relative path is "/")
+			FString CurrentDirectory(RelativeDirectoryFromMount);
+			FString UnusedCleanFileName;
+			while (!CurrentDirectory.IsEmpty())
+			{
+				// FPaths::GetPath doesn't handle our / at the end of directories, so call FindLastChar(/)
+				if (!SplitPathInline(CurrentDirectory, UnusedCleanFileName))
+				{
+					break;
+				}
+				DirectoryTreeIndex->FindOrAdd(CurrentDirectory);
+			}
+
+			// Add the new directory; this has to come after the addition of the parent directories because we want to use the pointer afterwards and adding other directories would invalidate it
+			Directory = &DirectoryTreeIndex->FindOrAdd(RelativeDirectoryFromMount);
+		}
+		Directory->Add(CleanFileName, EntryLocation);
+	}
+#endif
 
 	// Add the entry into the PathHash index
 	if (CollisionDetection || PathHashIndex)
@@ -1461,8 +1752,9 @@ const FPakDirectory* FPakFile::FindPrunedDirectoryInternal(const FString& Relati
 #if ENABLE_PAKFILE_RUNTIME_PRUNING_VALIDATE
 	if (ShouldValidatePrunedDirectory())
 	{
-		PakDirectory = DirectoryIndex.Find(RelativePathFromMount);
-		const FPakDirectory* PrunedPakDirectory = PrunedDirectoryIndex.Find(RelativePathFromMount);
+		PakDirectory = FindPrunedDirectoryInIndexInternal(RelativePathFromMount, DirectoryIndex, DirectoryTreeIndex);
+		const FPakDirectory* PrunedPakDirectory = FindPrunedDirectoryInIndexInternal(RelativePathFromMount,
+			PrunedDirectoryIndex, PrunedDirectoryTreeIndex);
 		if ((PakDirectory != nullptr) != (PrunedPakDirectory != nullptr))
 		{
 			TSet<FString> FullFoundFiles, PrunedFoundFiles;
@@ -1481,7 +1773,36 @@ const FPakDirectory* FPakFile::FindPrunedDirectoryInternal(const FString& Relati
 	else
 #endif
 	{
-		PakDirectory = DirectoryIndex.Find(RelativePathFromMount);
+		PakDirectory = FindPrunedDirectoryInIndexInternal(RelativePathFromMount, DirectoryIndex, DirectoryTreeIndex);
+	}
+	return PakDirectory;
+}
+
+const FPakDirectory* FPakFile::FindPrunedDirectoryInIndexInternal(const FString& RelativePathFromMount,
+	const FDirectoryIndex& InDirectoryIndex, const FDirectoryTreeIndex& InTreeIndex) const
+{
+	const FPakDirectory* PakDirectory = nullptr;
+#if ENABLE_PAKFILE_USE_DIRECTORY_TREE
+	if (ShouldUseDirectoryTree())
+	{
+		PakDirectory = InTreeIndex.Find(RelativePathFromMount);
+#if !UE_BUILD_SHIPPING
+		if (GPak_ValidateDirectoryTreeSearchConsistency)
+		{
+			const FPakDirectory* IndexedResult = InDirectoryIndex.Find(RelativePathFromMount);
+			if ((PakDirectory != nullptr) != (IndexedResult != nullptr) ||
+				(PakDirectory && PakDirectory->Num() != IndexedResult->Num()))
+			{
+				UE_LOG(LogPakFile, Fatal, TEXT("Mismatch between directoryindex and directorytreeindex search when searching for [%s] in pak [%s]"),
+					*FString(RelativePathFromMount), *GetFilename());
+			}
+		}
+#endif // !UE_BUILD_SHIPPING
+	}
+	else
+#endif // ENABLE_PAKFILE_USE_DIRECTORY_TREE
+	{
+		PakDirectory = InDirectoryIndex.Find(RelativePathFromMount);
 	}
 	return PakDirectory;
 }
@@ -1866,6 +2187,49 @@ const FPakEntryLocation* FPakFile::FindLocationFromIndex(const FString& FullPath
 	return nullptr;
 }
 
+const FPakEntryLocation* FPakFile::FindLocationFromIndex(const FString& FullPath,
+	const FDirectoryIndex& InDirectoryIndex, const FDirectoryTreeIndex& InDirectoryTreeIndex) const
+{
+#if ENABLE_PAKFILE_USE_DIRECTORY_TREE
+	if (ShouldUseDirectoryTree())
+	{
+		const FPakEntryLocation* Result = nullptr;
+		if (FullPath.StartsWith(MountPoint))
+		{
+			const TCHAR* RelativePathFromMount = (*FullPath) + MountPoint.Len();
+			FString RelativeDirName(RelativePathFromMount);
+			FString CleanFileName;
+			if (!RelativeDirName.IsEmpty())
+			{
+				SplitPathInline(RelativeDirName, CleanFileName);
+				const FPakDirectory* PakDirectory = InDirectoryTreeIndex.Find(RelativeDirName);
+				if (PakDirectory)
+				{
+					Result = PakDirectory->Find(CleanFileName);
+				}
+			}
+		}
+
+#if !UE_BUILD_SHIPPING
+		if (GPak_ValidateDirectoryTreeSearchConsistency)
+		{
+			const FPakEntryLocation* IndexedResult = FindLocationFromIndex(FullPath, MountPoint, InDirectoryIndex);
+			if ((Result != nullptr) != (IndexedResult != nullptr))
+			{
+				UE_LOG(LogPakFile, Fatal, TEXT("Mismatch between directoryindex and directorytreeindex search when searching for [%s] in pak [%s]"),
+					*FullPath, *GetFilename());
+			}
+		}
+#endif // !UE_BUILD_SHIPPING
+		return Result;
+	}
+	else
+#endif // ENABLE_PAKFILE_USE_DIRECTORY_TREE
+	{
+		return FindLocationFromIndex(FullPath, MountPoint, InDirectoryIndex);
+	}
+}
+
 FPakFile::EFindResult FPakFile::Find(const FString& FullPath, FPakEntry* OutEntry) const
 {
 	//QUICK_SCOPE_CYCLE_COUNTER(PakFileFind);
@@ -1881,7 +2245,7 @@ FPakFile::EFindResult FPakFile::Find(const FString& FullPath, FPakEntry* OutEntr
 
 		{
 			FScopedPakDirectoryIndexAccess ScopeAccess(*this);
-			DirectoryLocation = FindLocationFromIndex(FullPath, MountPoint, DirectoryIndex);
+			DirectoryLocation = FindLocationFromIndex(FullPath, DirectoryIndex, DirectoryTreeIndex);
 		}
 
 		if ((PathHashLocation != nullptr) != (DirectoryLocation != nullptr))
@@ -1906,9 +2270,12 @@ FPakFile::EFindResult FPakFile::Find(const FString& FullPath, FPakEntry* OutEntr
 		}
 		else
 		{
-			check(bHasFullDirectoryIndex);
+			// When we are using a pruned directory index, and no pathhash, calling Find with a filename will fail if the filename
+			// was pruned. Therefore we don't support calling it in pakfiles with a pruned directory index. But it's okay to call on
+			// an empty pakfile since that would return NotFound for every file.
+			check(Files.IsEmpty() || bHasFullDirectoryIndex);
 			FScopedPakDirectoryIndexAccess ScopeAccess(*this);
-			PakEntryLocation = FindLocationFromIndex(FullPath, MountPoint, DirectoryIndex);
+			PakEntryLocation = FindLocationFromIndex(FullPath, DirectoryIndex, DirectoryTreeIndex);
 		}
 	}
 	if (!PakEntryLocation)
@@ -1948,5 +2315,302 @@ void FPakFile::AddSpecialFile(const FPakEntry& Entry, const FString& Filename)
 	}
 
 	FPathHashIndex* PathHashToWrite = bHasPathHashIndex ? &PathHashIndex : nullptr;
-	AddEntryToIndex(Filename, EntryLocation, MountPoint, PathHashSeed, &DirectoryIndex, PathHashToWrite, nullptr /* CollisionDetection */, Info.Version);
+#if ENABLE_PAKFILE_USE_DIRECTORY_TREE
+	if (ShouldUseDirectoryTree())
+	{
+		AddEntryToIndex(Filename, EntryLocation, MountPoint, PathHashSeed,
+#if !UE_BUILD_SHIPPING
+			GPak_ValidateDirectoryTreeSearchConsistency ? &DirectoryIndex : nullptr,
+#else
+			nullptr,
+#endif // !UE_BUILD_SHIPPING
+			&DirectoryTreeIndex,
+			PathHashToWrite, nullptr /* CollisionDetection */, Info.Version);
+	}
+	else
+#endif // ENABLE_PAKFILE_USE_DIRECTORY_TREE
+	{
+		AddEntryToIndex(Filename, EntryLocation, MountPoint, PathHashSeed, &DirectoryIndex, nullptr,
+			PathHashToWrite, nullptr /* CollisionDetection */, Info.Version);
+	}
+}
+
+FPakFile::FBaseIterator& FPakFile::FBaseIterator::operator++()
+{
+	switch (IteratorType)
+	{
+	case EIteratorType::PathHash:
+		++GetPathHashIt();
+		break;
+	case EIteratorType::DirectoryTree:
+		[[fallthrough]];
+	case EIteratorType::DirectoryIndex:
+		++GetFileIt();
+		break;
+	default:
+		checkNoEntry();
+		break;
+	}
+	AdvanceToValid();
+	return *this;
+}
+
+FPakFile::FBaseIterator::operator bool() const
+{
+	switch (IteratorType)
+	{
+	case EIteratorType::PathHash:
+		return (bool)GetPathHashIt();
+	case EIteratorType::DirectoryTree:
+#if ENABLE_PAKFILE_USE_DIRECTORY_TREE
+		return (bool)GetTreeIt();
+#else
+		check(false);
+		return false;
+#endif
+	case EIteratorType::DirectoryIndex:
+		return (bool)GetDirectoryIndexIt();
+	default:
+		checkNoEntry();
+		return false;
+	}
+}
+
+bool FPakFile::FBaseIterator::operator !() const
+{
+	return !(bool)*this;
+}
+
+const FPakEntry& FPakFile::FBaseIterator::Info() const
+{
+	PakFile->GetPakEntry(GetPakEntryIndex(), &PakEntry);
+	return PakEntry;
+}
+
+bool FPakFile::FBaseIterator::HasFilename() const
+{
+	return (IteratorType == EIteratorType::DirectoryIndex) | (IteratorType == EIteratorType::DirectoryTree);
+}
+
+FPakFile::FBaseIterator::FBaseIterator(const FPakFile& InPakFile, bool bInIncludeDeleted, bool bInUsePathHash)
+	: PakFile(&InPakFile)
+	, bIncludeDeleted(bInIncludeDeleted)
+#if ENABLE_PAKFILE_RUNTIME_PRUNING
+	, bRequiresDirectoryIndexLock(false)
+#endif
+{
+	if (bInUsePathHash)
+	{
+		check(PakFile->bHasPathHashIndex);
+		IteratorType = EIteratorType::PathHash;
+		PathHashIt.Emplace(PakFile->PathHashIndex);
+	}
+	else
+	{
+#if ENABLE_PAKFILE_RUNTIME_PRUNING
+		bRequiresDirectoryIndexLock = PakFile->RequiresDirectoryIndexLock();
+		if (bRequiresDirectoryIndexLock)
+		{
+			PakFile->DirectoryIndexLock.ReadLock();
+		}
+#endif
+#if ENABLE_PAKFILE_USE_DIRECTORY_TREE
+		if (PakFile->ShouldUseDirectoryTree())
+		{
+			IteratorType = EIteratorType::DirectoryTree;
+			TreeIt.Reset(new FDirectoryTreeIndex::FConstIterator(PakFile->DirectoryTreeIndex.CreateConstIterator()));
+		}
+		else
+#endif
+		{
+			IteratorType = EIteratorType::DirectoryIndex;
+			DirectoryIndexIt.Emplace(PakFile->DirectoryIndex);
+		}
+		if (IsDirectoryItValid())
+		{
+			FileIt.Emplace(GetDirectoryItValue());
+		}
+	}
+
+	AdvanceToValid();
 }
+
+#if ENABLE_PAKFILE_RUNTIME_PRUNING
+FPakFile::FBaseIterator::~FBaseIterator()
+{
+	if (bRequiresDirectoryIndexLock)
+	{
+		PakFile->DirectoryIndexLock.ReadUnlock();
+	}
+}
+#endif
+
+const FString& FPakFile::FBaseIterator::Filename() const
+{
+	if (IteratorType == EIteratorType::PathHash)
+	{
+		// Filenames are not supported, CachedFilename is always empty 
+	}
+	else
+	{
+		checkf((bool)*this, TEXT("It is not legal to call Filename() on an invalid iterator"));
+		if (CachedFilename.IsEmpty())
+		{
+			CachedFilename = PakPathCombine(GetDirectoryItKey(), GetFileIt()->Key);
+		}
+	}
+	return CachedFilename;
+}
+
+FPakEntryLocation FPakFile::FBaseIterator::GetPakEntryIndex() const
+{
+	switch (IteratorType)
+	{
+	case EIteratorType::PathHash:
+		return GetPathHashIt().Value();
+	case EIteratorType::DirectoryTree:
+		[[fallthrough]];
+	case EIteratorType::DirectoryIndex:
+		return GetFileIt().Value();
+	default:
+		checkNoEntry();
+		return FPakEntryLocation();
+	}
+}
+
+void FPakFile::FBaseIterator::AdvanceToValid()
+{
+	if (IteratorType == EIteratorType::PathHash)
+	{
+		while (GetPathHashIt() && !bIncludeDeleted && Info().IsDeleteRecord())
+		{
+			++GetPathHashIt();
+		}
+	}
+	else
+	{
+		check((IteratorType == EIteratorType::DirectoryTree) | (IteratorType == EIteratorType::DirectoryIndex));
+		while (IsDirectoryItValid() && (!GetFileIt() || (!bIncludeDeleted && Info().IsDeleteRecord())))
+		{
+			if (GetFileIt())
+			{
+				++GetFileIt();
+			}
+			else
+			{
+				// No more files in the current directory, jump to the next one.
+				IncrementDirectoryIt();
+				if (IsDirectoryItValid())
+				{
+					FileIt.Emplace(GetDirectoryItValue());
+				}
+			}
+		}
+		CachedFilename.Reset();
+	}
+}
+
+FPakFile::FPathHashIndex::TConstIterator& FPakFile::FBaseIterator::GetPathHashIt()
+{
+	check(IteratorType == EIteratorType::PathHash);
+	return *PathHashIt;
+}
+const FPakFile::FPathHashIndex::TConstIterator& FPakFile::FBaseIterator::GetPathHashIt() const
+{
+	return const_cast<FBaseIterator*>(this)->GetPathHashIt();
+}
+
+bool FPakFile::FBaseIterator::IsDirectoryItValid() const
+{
+	check((IteratorType == EIteratorType::DirectoryIndex) | (IteratorType == EIteratorType::DirectoryTree));
+#if ENABLE_PAKFILE_USE_DIRECTORY_TREE
+	if (IteratorType == EIteratorType::DirectoryTree)
+	{
+		return (bool)GetTreeIt();
+	}
+	else
+#endif
+	{
+		return (bool)GetDirectoryIndexIt();
+	}
+}
+
+void FPakFile::FBaseIterator::IncrementDirectoryIt()
+{
+	check((IteratorType == EIteratorType::DirectoryIndex) | (IteratorType == EIteratorType::DirectoryTree));
+#if ENABLE_PAKFILE_USE_DIRECTORY_TREE
+	if (IteratorType == EIteratorType::DirectoryTree)
+	{
+		++GetTreeIt();
+	}
+	else
+#endif
+	{
+		++GetDirectoryIndexIt();
+	}
+}
+
+FStringView FPakFile::FBaseIterator::GetDirectoryItKey() const
+{
+	check((IteratorType == EIteratorType::DirectoryIndex) | (IteratorType == EIteratorType::DirectoryTree));
+#if ENABLE_PAKFILE_USE_DIRECTORY_TREE
+	if (IteratorType == EIteratorType::DirectoryTree)
+	{
+		return GetTreeIt()->Key;
+	}
+	else
+#endif
+	{
+		return GetDirectoryIndexIt()->Key;
+	}
+}
+
+const FPakDirectory& FPakFile::FBaseIterator::GetDirectoryItValue() const
+{
+	check((IteratorType == EIteratorType::DirectoryIndex) | (IteratorType == EIteratorType::DirectoryTree));
+#if ENABLE_PAKFILE_USE_DIRECTORY_TREE
+	if (IteratorType == EIteratorType::DirectoryTree)
+	{
+		return GetTreeIt()->Value;
+	}
+	else
+#endif
+	{
+		return GetDirectoryIndexIt()->Value;
+	}
+}
+
+FPakDirectory::TConstIterator& FPakFile::FBaseIterator::GetFileIt()
+{
+	check((IteratorType == EIteratorType::DirectoryIndex) | (IteratorType == EIteratorType::DirectoryTree));
+	return *FileIt;
+}
+
+const FPakDirectory::TConstIterator& FPakFile::FBaseIterator::GetFileIt() const
+{
+	return const_cast<FBaseIterator*>(this)->GetFileIt();
+}
+
+FPakFile::FDirectoryIndex::TConstIterator& FPakFile::FBaseIterator::GetDirectoryIndexIt()
+{
+	check(IteratorType == EIteratorType::DirectoryIndex);
+	return *DirectoryIndexIt;
+}
+
+const FPakFile::FDirectoryIndex::TConstIterator& FPakFile::FBaseIterator::GetDirectoryIndexIt() const
+{
+	return const_cast<FBaseIterator*>(this)->GetDirectoryIndexIt();
+}
+
+#if ENABLE_PAKFILE_USE_DIRECTORY_TREE
+FPakFile::FDirectoryTreeIndex::FConstIterator& FPakFile::FBaseIterator::GetTreeIt()
+{
+	check(IteratorType == EIteratorType::DirectoryTree);
+	return *TreeIt;
+}
+
+const FPakFile::FDirectoryTreeIndex::FConstIterator& FPakFile::FBaseIterator::GetTreeIt() const
+{
+	return const_cast<FBaseIterator*>(this)->GetTreeIt();
+}
+#endif
diff --git a/Engine/Source/Runtime/PakFile/Private/PakFile.inl b/Engine/Source/Runtime/PakFile/Private/PakFile.inl
--- a/Engine/Source/Runtime/PakFile/Private/PakFile.inl
+++ b/Engine/Source/Runtime/PakFile/Private/PakFile.inl
@@ -9,22 +9,22 @@
  */
 template <typename ShouldVisitFunc, class ContainerType>
 void FPakFile::FindFilesAtPathInIndex(const FDirectoryIndex& TargetIndex, const FDirectoryTreeIndex& TargetTreeIndex, ContainerType& OutFiles,
-	const FString& Directory, const FVisitFilter<ShouldVisitFunc>& VisitFilter) const
+	const FString& FullSearchPath, const FVisitFilter<ShouldVisitFunc>& VisitFilter) const
 {
 	TRACE_CPUPROFILER_EVENT_SCOPE(FindFilesAtPathInIndex);
 
-	FStringView RelativeSearch;
-	if (Directory.StartsWith(MountPoint))
+	FStringView RelSearchPath;
+	if (FullSearchPath.StartsWith(MountPoint))
 	{
-		RelativeSearch = FStringView(Directory).RightChop(MountPoint.Len());
+		RelSearchPath = FStringView(FullSearchPath).RightChop(MountPoint.Len());
 	}
 	else
 	{
 		// Directory is unnormalized and might not end with /; MountPoint is guaranteed to end with /.
 		// Act as if we were called with a normalized directory if adding slash makes it match MountPoint.
-		if (FStringView(Directory).StartsWith(FStringView(MountPoint).LeftChop(1)))
+		if (FStringView(FullSearchPath).StartsWith(FStringView(MountPoint).LeftChop(1)))
 		{
-			RelativeSearch = FStringView();
+			RelSearchPath = FStringView();
 		}
 		else
 		{
@@ -36,20 +36,21 @@ void FPakFile::FindFilesAtPathInIndex(const FDirectoryIndex& TargetIndex, const
 	TArray<FString> DirectoriesInPak; // List of all unique directories at path
 
 #if ENABLE_PAKFILE_USE_DIRECTORY_TREE
-	// DirectoryTree does not work with an empty RelativePath as SearchIndex. Fallback to the index search in this case as there is no performance benefit from the DirectoryTree anyway.
-	if (GPak_UseDirectoryTreeForPakSearch && !RelativeSearch.IsEmpty())
+	if (ShouldUseDirectoryTree())
 	{
-		FindFilesAtPathInTreeIndexInternal(FString(RelativeSearch), TargetTreeIndex, OutFiles, DirectoriesInPak, Directory, VisitFilter);
+		FindFilesAtPathInTreeIndexInternal(RelSearchPath, TargetTreeIndex, OutFiles, DirectoriesInPak, FullSearchPath, VisitFilter);
 		
 #if !UE_BUILD_SHIPPING
 		if (GPak_ValidateDirectoryTreeSearchConsistency)
 		{
 			ContainerType OutFilesIndexed;
 			TArray<FString> OutDirectoriesInPakIndexed;
-			FindFilesAtPathInIndexInternal(RelativeSearch, TargetIndex, OutFilesIndexed, OutDirectoriesInPakIndexed, Directory, VisitFilter);
+			FindFilesAtPathInIndexInternal(RelSearchPath, TargetIndex, OutFilesIndexed, OutDirectoriesInPakIndexed, FullSearchPath, VisitFilter);
 			if (!ValidateDirectoryTreeSearchConsistency(OutFiles, DirectoriesInPak, OutFilesIndexed, OutDirectoriesInPakIndexed))
 			{
-				UE_LOG(LogPakFile, Fatal, TEXT("Mismatch between directoryindex and directorytreeindex search when searching for [%s] in pak [%s]"), *FString(RelativeSearch), *GetFilename());
+				UE_LOG(LogPakFile, Fatal,
+					TEXT("Mismatch between directoryindex and directorytreeindex search when searching for [%s] in pak [%s]"),
+					*FString(RelSearchPath), *GetFilename());
 			}
 		}
 #endif	// !UE_BUILD_SHIPPING
@@ -57,33 +58,49 @@ void FPakFile::FindFilesAtPathInIndex(const FDirectoryIndex& TargetIndex, const
 	else
 #endif	// ENABLE_PAKFILE_USE_DIRECTORY_TREE
 	{
-		FindFilesAtPathInIndexInternal(RelativeSearch, TargetIndex, OutFiles, DirectoriesInPak, Directory, VisitFilter);
+		FindFilesAtPathInIndexInternal(RelSearchPath, TargetIndex, OutFiles, DirectoriesInPak, FullSearchPath, VisitFilter);
 	}
 	OutFiles.Append(MoveTemp(DirectoriesInPak));
 }
 
 #if ENABLE_PAKFILE_USE_DIRECTORY_TREE
 template <typename ShouldVisitFunc, class ContainerType>
-void FPakFile::FindFilesAtPathInTreeIndexInternal(const FString& Path, const FDirectoryTreeIndex& TargetTreeIndex, ContainerType& OutFiles, TArray<FString>& OutDirectories,
-	const FString& Directory, const FVisitFilter<ShouldVisitFunc>& VisitFilter) const
+void FPakFile::FindFilesAtPathInTreeIndexInternal(FStringView RelSearchPath, const FDirectoryTreeIndex& TargetTreeIndex,
+	ContainerType& OutFiles, TArray<FString>& OutDirectories, const FString& FullSearchPath,
+	const FVisitFilter<ShouldVisitFunc>& VisitFilter) const
 {
-	check(!Path.IsEmpty());
-	FPakDirectory* const* PakDirectory = TargetTreeIndex.Find(Path);
-	if (PakDirectory)
+	if (RelSearchPath.IsEmpty())
 	{
-		FindFilesAtPathInPakDirectoryInternal(MountPoint, Path, **PakDirectory, OutFiles, OutDirectories, Directory, VisitFilter);
-
-		if (VisitFilter.bRecursive || VisitFilter.bIncludeDirectories)
+		for (TPair<FStringView, const FPakDirectory&> Pair : TargetTreeIndex)
 		{
-			TArray<FString> OutChildDirectories;
-			TargetTreeIndex.TryGetChildren(Path, OutChildDirectories, VisitFilter.bRecursive ? EDirectoryTreeGetFlags::Recursive : EDirectoryTreeGetFlags::None);
-			for (const FString& ChildDirectoryPath : OutChildDirectories)
+			FindFilesAtPathInPakDirectoryInternal(MountPoint, Pair.Key, Pair.Value, OutFiles, OutDirectories,
+				FullSearchPath, VisitFilter);
+		}
+	}
+	else
+	{
+		const FPakDirectory* PakDirectory = TargetTreeIndex.Find(RelSearchPath);
+		if (PakDirectory)
+		{
+			FindFilesAtPathInPakDirectoryInternal(MountPoint, RelSearchPath, *PakDirectory, OutFiles, OutDirectories,
+				FullSearchPath, VisitFilter);
+
+			if (VisitFilter.bRecursive || VisitFilter.bIncludeDirectories)
 			{
-				FString ChildDirectoryPathFull = PakPathCombine(FString(Path), ChildDirectoryPath);
-				MakeDirectoryFromPath(ChildDirectoryPathFull);
-				if (FPakDirectory* const* PakDirectoryChild = TargetTreeIndex.Find(ChildDirectoryPathFull))
+				// TODO: Add TDirectoryTree::CreatePathIterator so we can avoid the copy of children and the redundant
+				// lookups of the children in the tree.
+				TArray<FString> OutChildDirectories;
+				TargetTreeIndex.TryGetChildren(RelSearchPath, OutChildDirectories,
+					VisitFilter.bRecursive ? EDirectoryTreeGetFlags::Recursive : EDirectoryTreeGetFlags::None);
+				for (const FString& ChildDirectoryPath : OutChildDirectories)
 				{
-					FindFilesAtPathInPakDirectoryInternal(MountPoint, ChildDirectoryPathFull, **PakDirectoryChild, OutFiles, OutDirectories, Directory, VisitFilter);
+					FString RelChildPath = PakPathCombine(RelSearchPath, ChildDirectoryPath);
+					MakeDirectoryFromPath(RelChildPath);
+					if (const FPakDirectory* PakDirectoryChild = TargetTreeIndex.Find(RelChildPath))
+					{
+						FindFilesAtPathInPakDirectoryInternal(MountPoint, RelChildPath,
+							*PakDirectoryChild, OutFiles, OutDirectories, FullSearchPath, VisitFilter);
+					}
 				}
 			}
 		}
@@ -131,68 +148,73 @@ bool FPakFile::ValidateDirectoryTreeSearchConsistency(const ContainerType& Files
 #endif // ENABLE_PAKFILE_USE_DIRECTORY_TREE
 
 template <typename ShouldVisitFunc, class ContainerType>
-void FPakFile::FindFilesAtPathInIndexInternal(const FStringView& Path, const FDirectoryIndex& TargetIndex, ContainerType& OutFiles, TArray<FString>& OutDirectories,
-	const FString& Directory, const FVisitFilter<ShouldVisitFunc>& VisitFilter) const
+void FPakFile::FindFilesAtPathInIndexInternal(const FStringView& RelSearchPath, const FDirectoryIndex& TargetIndex,
+	ContainerType& OutFiles, TArray<FString>& OutDirectories,
+	const FString& FullSearchPath, const FVisitFilter<ShouldVisitFunc>& VisitFilter) const
 {
 	for (TMap<FString, FPakDirectory>::TConstIterator It(TargetIndex); It; ++It)
 	{
 		// Check if the file is under the specified path.
-		if (FStringView(It.Key()).StartsWith(Path))
+		if (FStringView(It.Key()).StartsWith(RelSearchPath))
 		{
-			FindFilesAtPathInPakDirectoryInternal(MountPoint, It.Key(), It.Value(), OutFiles, OutDirectories, Directory, VisitFilter);
+			FindFilesAtPathInPakDirectoryInternal(MountPoint, It.Key(), It.Value(), OutFiles, OutDirectories,
+				FullSearchPath, VisitFilter);
 		}
 	}
 }
 
 template <typename ShouldVisitFunc, class ContainerType>
-void FPakFile::FindFilesAtPathInPakDirectoryInternal(const FString& MountPoint, const FString& Path, const FPakDirectory& PakDirectory, ContainerType& OutFiles, TArray<FString>& OutDirectoriesInPak,
-	const FString& Directory, const FVisitFilter<ShouldVisitFunc>& VisitFilter)
+void FPakFile::FindFilesAtPathInPakDirectoryInternal(const FString& MountPoint, FStringView RelPathInIndex,
+	const FPakDirectory& PakDirectory, ContainerType& OutFiles, TArray<FString>& OutDirectoriesInPak,
+	const FString& FullSearchPath, const FVisitFilter<ShouldVisitFunc>& VisitFilter)
 {
-	FString PakPath = PakPathCombine(MountPoint, Path);
+	FString FullPathInIndex = PakPathCombine(MountPoint, RelPathInIndex);
 	if (VisitFilter.bRecursive == true)
 	{
 		// Add everything
 		if (VisitFilter.bIncludeFiles)
 		{
-			for (FPakDirectory::TConstIterator DirectoryIt(PakDirectory); DirectoryIt; ++DirectoryIt)
+			for (FPakDirectory::TConstIterator FileIt(PakDirectory); FileIt; ++FileIt)
 			{
-				const FString& FilePathUnderDirectory = DirectoryIt.Key();
+				const FString& FilePathUnderDirectory = FileIt.Key();
 				if (VisitFilter.ShouldVisit(FilePathUnderDirectory))
 				{
-					OutFiles.Add(PakPathCombine(PakPath, FilePathUnderDirectory));
+					OutFiles.Add(PakPathCombine(FullPathInIndex, FilePathUnderDirectory));
 				}
 			}
 		}
 		if (VisitFilter.bIncludeDirectories)
 		{
-			if (Directory != PakPath)
+			if (FullSearchPath != FullPathInIndex)
 			{
-				if (VisitFilter.ShouldVisit(PakPath))
+				if (VisitFilter.ShouldVisit(FullPathInIndex))
 				{
-					OutDirectoriesInPak.Add(MoveTemp(PakPath));
+					OutDirectoriesInPak.Add(MoveTemp(FullPathInIndex));
 				}
 			}
 		}
 	}
 	else
 	{
-		int32 SubDirIndex = PakPath.Len() > Directory.Len() ? PakPath.Find(TEXT("/"), ESearchCase::CaseSensitive, ESearchDir::FromStart, Directory.Len() + 1) : INDEX_NONE;
+		int32 SubDirIndex = FullPathInIndex.Len() > FullSearchPath.Len()
+			? FullPathInIndex.Find(TEXT("/"), ESearchCase::CaseSensitive, ESearchDir::FromStart, FullSearchPath.Len() + 1)
+			: INDEX_NONE;
 		// Add files in the specified folder only.
 		if (VisitFilter.bIncludeFiles && SubDirIndex == INDEX_NONE)
 		{
-			for (FPakDirectory::TConstIterator DirectoryIt(PakDirectory); DirectoryIt; ++DirectoryIt)
+			for (FPakDirectory::TConstIterator FileIt(PakDirectory); FileIt; ++FileIt)
 			{
-				const FString& FilePathUnderDirectory = DirectoryIt.Key();
+				const FString& FilePathUnderDirectory = FileIt.Key();
 				if (VisitFilter.ShouldVisit(FilePathUnderDirectory))
 				{
-					OutFiles.Add(PakPathCombine(PakPath, FilePathUnderDirectory));
+					OutFiles.Add(PakPathCombine(FullPathInIndex, FilePathUnderDirectory));
 				}
 			}
 		}
 		// Add sub-folders in the specified folder only
 		if (VisitFilter.bIncludeDirectories && SubDirIndex >= 0)
 		{
-			FString SubDirPath = PakPath.Left(SubDirIndex + 1);
+			FString SubDirPath = FullPathInIndex.Left(SubDirIndex + 1);
 			if (VisitFilter.ShouldVisit(SubDirPath))
 			{
 				OutDirectoriesInPak.AddUnique(MoveTemp(SubDirPath));
@@ -215,13 +237,13 @@ template <typename ShouldVisitFunc, class ContainerType>
 void FPakFile::FindPrunedFilesAtPathInternal(const TCHAR* InPath, ContainerType& OutFiles, const FVisitFilter<ShouldVisitFunc>& VisitFilter) const
 {
 	// Make sure all directory names end with '/'.
-	FString Directory(InPath);
-	MakeDirectoryFromPath(Directory);
+	FString FullSearchPath(InPath);
+	MakeDirectoryFromPath(FullSearchPath);
 
 	// Check the specified path is under the mount point of this pak file.
 	// The reverse case (MountPoint StartsWith Directory) is needed to properly handle
 	// pak files that are a subdirectory of the actual directory.
-	if (!Directory.StartsWith(MountPoint) && !MountPoint.StartsWith(Directory))
+	if (!FullSearchPath.StartsWith(MountPoint) && !MountPoint.StartsWith(FullSearchPath))
 	{
 		return;
 	}
@@ -231,8 +253,8 @@ void FPakFile::FindPrunedFilesAtPathInternal(const TCHAR* InPath, ContainerType&
 	if (ShouldValidatePrunedDirectory())
 	{
 		TSet<FString> FullFoundFiles, PrunedFoundFiles;
-		FindFilesAtPathInIndex(DirectoryIndex, DirectoryTreeIndex, FullFoundFiles, Directory, VisitFilter);
-		FindFilesAtPathInIndex(PrunedDirectoryIndex, DirectoryTreeIndex, PrunedFoundFiles, Directory, VisitFilter);
+		FindFilesAtPathInIndex(DirectoryIndex, DirectoryTreeIndex, FullFoundFiles, FullSearchPath, VisitFilter);
+		FindFilesAtPathInIndex(PrunedDirectoryIndex, PrunedDirectoryTreeIndex, PrunedFoundFiles, FullSearchPath, VisitFilter);
 		ValidateDirectorySearch(FullFoundFiles, PrunedFoundFiles, InPath);
 
 		for (const FString& FoundFile : FullFoundFiles)
@@ -243,6 +265,6 @@ void FPakFile::FindPrunedFilesAtPathInternal(const TCHAR* InPath, ContainerType&
 	else
 #endif
 	{
-		FindFilesAtPathInIndex(DirectoryIndex, DirectoryTreeIndex, OutFiles, Directory, VisitFilter);
+		FindFilesAtPathInIndex(DirectoryIndex, DirectoryTreeIndex, OutFiles, FullSearchPath, VisitFilter);
 	}
 }
diff --git a/Engine/Source/Runtime/PakFile/Public/IPlatformFilePak.h b/Engine/Source/Runtime/PakFile/Public/IPlatformFilePak.h
--- a/Engine/Source/Runtime/PakFile/Public/IPlatformFilePak.h
+++ b/Engine/Source/Runtime/PakFile/Public/IPlatformFilePak.h
@@ -32,6 +32,7 @@
 #include "Misc/DateTime.h"
 #include "Misc/Guid.h"
 #include "Misc/IEngineCrypto.h"
+#include "Misc/Optional.h"
 #include "Misc/Paths.h"
 #include "Misc/ScopeLock.h"
 #include "Misc/ScopeRWLock.h"
@@ -715,6 +716,18 @@ public:
 /** Pak directory type mapping a filename to an FPakEntryLocation. */
 typedef TMap<FString, FPakEntryLocation> FPakDirectory;
 
+namespace UE::PakFile::Private
+{
+#if ENABLE_PAKFILE_USE_DIRECTORY_TREE
+	typedef TDirectoryTree<FPakDirectory> FDirectoryTreeIndex;
+#else
+	struct FNullDirectoryTreeIndex
+	{
+	};
+	typedef FNullDirectoryTreeIndex FDirectoryTreeIndex;
+#endif
+}
+
 /* Convenience struct for building FPakFile indexes from an enumeration of (Filename,FPakEntry) pairs */
 struct FPakEntryPair
 {
@@ -801,20 +814,14 @@ private:
 	/** Pak Index organized as a map of directories to support searches by path.  This Index is pruned at runtime of all FileNames and Paths that are not allowed by DirectoryIndexKeepFiles */
 	FDirectoryIndex DirectoryIndex;
 
-#if ENABLE_PAKFILE_USE_DIRECTORY_TREE
-	typedef TDirectoryTree<FPakDirectory*> FDirectoryTreeIndex;
-#else
-	struct FNullDirectoryTreeIndex
-	{
-	};
-	typedef FNullDirectoryTreeIndex FDirectoryTreeIndex;
-#endif
+	using FDirectoryTreeIndex = UE::PakFile::Private::FDirectoryTreeIndex;
 
 	FDirectoryTreeIndex DirectoryTreeIndex;
 
 #if ENABLE_PAKFILE_RUNTIME_PRUNING
 	/** Temporary-lifetime copy of the Pruned DirectoryIndex; all Pruned files have been removed form this copy.  This copy is used for validation that no queries are missing during runtime, and will be swapped into the DirectoryIndex when Pak Mounting is complete */
 	FDirectoryIndex PrunedDirectoryIndex;
+	FDirectoryTreeIndex PrunedDirectoryTreeIndex;
 	/** ReaderWriter lock to guard DirectoryIndex iteration from being interrupted by the swap of PrunedDirectoryIndex */
 	mutable FRWLock DirectoryIndexLock;
 #endif
@@ -845,6 +852,9 @@ private:
 	bool bWillPruneDirectoryIndex;
 	/* True if the Index of this PakFile was a legacy index that did not have the precomputed Pruned DirectoryIndex and we need to compute it before swapping the Pruned DirectoryIndex*/
 	bool bNeedsLegacyPruning;
+#endif
+#if ENABLE_PAKFILE_USE_DIRECTORY_TREE
+	bool bUseDirectoryTree;
 #endif
 	/** ID for the chunk this pakfile is part of. INDEX_NONE if this isn't a pak chunk (derived from filename) */
 	int32 PakchunkIndex;
@@ -1130,24 +1140,39 @@ public:
 	 */
 	PAKFILE_API bool Check();
 
+	/** Internal IteratorType used by an FBaseIterator. */
+	enum class EIteratorType : uint8
+	{
+		PathHash,
+		DirectoryIndex,
+		DirectoryTree,
+	};
 	/** Base functionality for iterating over the DirectoryIndex. */
 	class FBaseIterator
 	{
 	private:
 		/** Owner pak file. */
-		const FPakFile& PakFile;
+		const FPakFile* PakFile;
 		/** Outer iterator over Directories when using the FDirectoryIndex. */
-		FDirectoryIndex::TConstIterator DirectoryIndexIt;
-		/** Inner iterator over Files when using the FDirectoryIndex. */
-		FPakDirectory::TConstIterator DirectoryIt;
+		TOptional<FDirectoryIndex::TConstIterator> DirectoryIndexIt;
+#if ENABLE_PAKFILE_USE_DIRECTORY_TREE
+		/**
+		 * Outer iterator over Directories when using the FDirectoryTreeIndex.
+		 * A TDirectoryTree::FConstIterator has a large struct size. Store it as a TUniquePtr rather than a TOptional
+		 * so that the struct size of our struct can remain small.
+		 */
+		TUniquePtr<FDirectoryTreeIndex::FConstIterator> TreeIt;
+#endif
+		/** Inner iterator over Files when using the FDirectoryIndex or FDirectoryTreeIndex. */
+		TOptional<FPakDirectory::TConstIterator> FileIt;
 		/** Iterator when using the FPathHashIndex. */
-		FPathHashIndex::TConstIterator PathHashIt;
+		TOptional<FPathHashIndex::TConstIterator> PathHashIt;
 		/** The cached filename for return in Filename(). */
 		mutable FString CachedFilename;
 		/* The PakEntry for return in Info */
 		mutable FPakEntry PakEntry;
-		/* Whether to use the PathHashIndex. */
-		bool bUsePathHash;
+		/* Which type of internal iterator this iterator and its pakfile are using. */
+		EIteratorType IteratorType;
 		/** Whether to include delete records in the iteration. */
 		bool bIncludeDeleted;
 #if ENABLE_PAKFILE_RUNTIME_PRUNING
@@ -1156,163 +1181,67 @@ public:
 #endif
 
 	public:
-		FBaseIterator& operator++()
-		{
-			if (bUsePathHash)
-			{
-				++PathHashIt;
-			}
-			else
-			{
-				++DirectoryIt;
-			}
-			AdvanceToValid();
-			return *this;
-		}
+		FBaseIterator(FBaseIterator&&) = default;
+		FBaseIterator& operator=(FBaseIterator&&) = default;
+		// PakFile iterators can be large, because they have to walk a tree, so do not copy them, move construct only.
+		FBaseIterator(const FBaseIterator&) = delete;
+		FBaseIterator& operator=(const FBaseIterator&) = delete;
+
+		PAKFILE_API FBaseIterator& operator++();
 
 		/** conversion to "bool" returning true if the iterator is valid. */
-		FORCEINLINE explicit operator bool() const
-		{
-			if (bUsePathHash)
-			{
-				return !!PathHashIt;
-			}
-			else
-			{
-				return !!DirectoryIndexIt;
-			}
-		}
+		PAKFILE_API explicit operator bool() const;
 
 		/** inverse of the "bool" operator */
-		FORCEINLINE bool operator !() const
-		{
-			return !(bool)*this;
-		}
+		PAKFILE_API bool operator !() const;
 
 		/** Return the FPakEntry. Invalid to call unless the iterator is currently valid. */
-		const FPakEntry& Info() const
-		{
-			PakFile.GetPakEntry(GetPakEntryIndex(), &PakEntry);
-			return PakEntry;
-		}
+		PAKFILE_API const FPakEntry& Info() const;
 
-		bool HasFilename() const
-		{
-			return !bUsePathHash;
-		}
+		PAKFILE_API bool HasFilename() const;
 
 	protected:
-		FBaseIterator(const FPakFile& InPakFile, bool bInIncludeDeleted, bool bInUsePathHash)
-			: PakFile(InPakFile)
-			, DirectoryIndexIt(FDirectoryIndex())
-			, DirectoryIt(FPakDirectory())
-			, PathHashIt(PakFile.PathHashIndex)
-			, bUsePathHash(bInUsePathHash)
-			, bIncludeDeleted(bInIncludeDeleted)
-#if ENABLE_PAKFILE_RUNTIME_PRUNING
-			, bRequiresDirectoryIndexLock(false)
-#endif
-		{
-			if (bUsePathHash)
-			{
-				check(PakFile.bHasPathHashIndex);
-			}
-			else
-			{
-#if ENABLE_PAKFILE_RUNTIME_PRUNING
-				bRequiresDirectoryIndexLock = PakFile.RequiresDirectoryIndexLock();
-				if (bRequiresDirectoryIndexLock)
-				{
-					PakFile.DirectoryIndexLock.ReadLock();
-				}
-#endif
-				DirectoryIndexIt.~TConstIterator();
-				new (&DirectoryIndexIt) FDirectoryIndex::TConstIterator(InPakFile.DirectoryIndex);
-				if (DirectoryIndexIt)
-				{
-					DirectoryIt.~TConstIterator();
-					new (&DirectoryIt) FPakDirectory::TConstIterator(DirectoryIndexIt.Value());
-				}
-			}
-
-			AdvanceToValid();
-		}
+		PAKFILE_API FBaseIterator(const FPakFile& InPakFile, bool bInIncludeDeleted, bool bUsePathHash);
 
 #if ENABLE_PAKFILE_RUNTIME_PRUNING
-		~FBaseIterator()
-		{
-			if (bRequiresDirectoryIndexLock)
-			{
-				PakFile.DirectoryIndexLock.ReadUnlock();
-			}
-		}
+		PAKFILE_API ~FBaseIterator();
 #endif
 
 		/** Return the current filename, as the RelativePath from the MountPoint.  Only available when using the FDirectoryIndex, otherwise always returns empty string. Invalid to call unless the iterator is currently valid. */
-		const FString& Filename() const
-		{
-			if (bUsePathHash)
-			{
-				// Filenames are not supported, CachedFilename is always empty 
-			}
-			else
-			{
-				checkf(DirectoryIndexIt && DirectoryIt, TEXT("It is not legal to call Filename() on an invalid iterator"));
-				if (CachedFilename.IsEmpty())
-				{
-					CachedFilename = PakPathCombine(DirectoryIndexIt.Key(), DirectoryIt.Key());
-				}
-			}
-			return CachedFilename;
-		}
+		PAKFILE_API const FString& Filename() const;
 
 		/** Return the arbitrary index of the iteration. Invalid to call unless the iterator is currently valid. */
-		const FPakEntryLocation GetPakEntryIndex() const
-		{
-			if (bUsePathHash)
-			{
-				return PathHashIt.Value();
-			}
-			else
-			{
-				return DirectoryIt.Value();
-			}
-		}
+		PAKFILE_API FPakEntryLocation GetPakEntryIndex() const;
 
 	private:
 
 		/* Skips over deleted records and moves to the next Directory in the DirectoryIndex when necessary. */
-		FORCEINLINE void AdvanceToValid()
-		{
-			if (bUsePathHash)
-			{
-				while (PathHashIt && !bIncludeDeleted && Info().IsDeleteRecord())
-				{
-					++PathHashIt;
-				}
-			}
-			else
-			{
-				while (DirectoryIndexIt && (!DirectoryIt || (!bIncludeDeleted && Info().IsDeleteRecord())))
-				{
-					if (DirectoryIt)
-					{
-						++DirectoryIt;
-					}
-					else
-					{
-						// No more files in the current directory, jump to the next one.
-						++DirectoryIndexIt;
-						if (DirectoryIndexIt)
-						{
-							DirectoryIt.~TConstIterator();
-							new (&DirectoryIt) FPakDirectory::TConstIterator(DirectoryIndexIt.Value());
-						}
-					}
-				}
-				CachedFilename.Reset();
-			}
-		}
+		void AdvanceToValid();
+
+		/** With EIteratorType::PathHash, we have only a single iterator. */
+		FPathHashIndex::TConstIterator& GetPathHashIt();
+		const FPathHashIndex::TConstIterator& GetPathHashIt() const;
+
+		/**
+		 * With EIteratorType::DirectoryIndex or EIteratorType::DirectoryTree, we have an outer iterator that
+		 * traverses directories in the tree, and an inner iterator that traverses files in the directory.
+		 */
+		bool IsDirectoryItValid() const;
+		void IncrementDirectoryIt();
+		FStringView GetDirectoryItKey() const;
+		const FPakDirectory& GetDirectoryItValue() const;
+		FPakDirectory::TConstIterator& GetFileIt();
+		const FPakDirectory::TConstIterator& GetFileIt() const;
+
+		/** With EIteratorType::DirectoryIndex, GetDirectoryIndexIt is the outer iterator of directories. */
+		FDirectoryIndex::TConstIterator& GetDirectoryIndexIt();
+		const FDirectoryIndex::TConstIterator& GetDirectoryIndexIt() const;
+
+#if ENABLE_PAKFILE_USE_DIRECTORY_TREE
+		/** With EIteratorType::DirectoryTree, GetTreeIt is the outer iterator of directories. */
+		FDirectoryTreeIndex::FConstIterator& GetTreeIt();
+		const FDirectoryTreeIndex::FConstIterator& GetTreeIt() const;
+#endif
 	};
 
 	/** Iterator class for every FPakEntry in the FPakFile, but does not provide filenames unless the PakFile has an unpruned DirectoryIndex. */
@@ -1474,22 +1403,26 @@ public:
 	}
 
 	/* Helper function to join two path strings that are in the PakPath format */
-	static FString PakPathCombine(const FString& Parent, const FString& Child)
+	static FString PakPathCombine(FStringView Parent, FStringView Child)
 	{
 		// Our paths are different than FPaths, because our dirs / at the end, and "/" is the relative path to the mountdirectory and should be mapped to the empty string when joining
 		check(Parent.Len() > 0 && Parent[Parent.Len() - 1] == TEXT('/'));
 		if (Parent.Len() == 1)
 		{
-			return Child;
+			return FString(Child);
 		}
 		else if (Child.Len() == 1 && Child[0] == TEXT('/'))
 		{
-			return Parent;
+			return FString(Parent);
 		}
 		else
 		{
 			check(Child.Len() == 0 || Child[0] != TEXT('/'));
-			return Parent + Child;
+			FString Result;
+			Result.Reserve(Parent.Len() + Child.Len());
+			Result += Parent;
+			Result += Child;
+			return Result;
 		}
 	}
 
@@ -1598,12 +1531,12 @@ private:
 	/**
 	 * Initializes the pak file.
 	 */
-	PAKFILE_API void Initialize(FArchive& Reader, bool bLoadIndex = true);
+	void Initialize(FArchive& Reader, bool bLoadIndex = true);
 
 	/**
 	 * Loads and initializes pak file index.
 	 */
-	PAKFILE_API void LoadIndex(FArchive& Reader);
+	void LoadIndex(FArchive& Reader);
 
 	/**
 	  * Returns the FPakEntry pointed to by the given FPakEntryLocation, forwards to the static GetPakEntry with data from *this
@@ -1613,24 +1546,24 @@ private:
 	/** Helper class to read IndexSettings from project delegate and commandline */
 	struct FIndexSettings;
 
-	static PAKFILE_API FIndexSettings& GetIndexSettings();
+	static FIndexSettings& GetIndexSettings();
 
 	/**
 	  * Returns the global,const flag for whether the current process should run directory queries on both the DirectoryIndex and the Pruned DirectoryIndex and log an error if they don't match.
 	  * Validation only occurs until the first call to OptimizeMemoryUsageForMountedPaks, after which the Full DirectoryIndex is dropped and there is nothing left to Validate
 	  * Has the same effect as IsPakDelayPruning, plus the addition of the error for any mismatches.
 	  */
-	static PAKFILE_API bool IsPakValidatePruning();
+	static bool IsPakValidatePruning();
 	/**
 	 * Returns the global,const flag for whether the current process should keep a copy of the Full DirectoryIndex around until OptimizeMemoryUsageForMountedPaks is called, so that systems can run
 	 * directory queries against the full index until then.
 	 * Note that validation will still occur if IsPakValidatePruning is true.
 	 */
-	static PAKFILE_API bool IsPakDelayPruning();
+	static bool IsPakDelayPruning();
 
 #if ENABLE_PAKFILE_RUNTIME_PRUNING
 	/** Global flag for whether a Pak has indicated it needs Pruning */
-	static PAKFILE_API bool bSomePakNeedsPruning;
+	static bool bSomePakNeedsPruning;
 #endif
 
 	/**
@@ -1644,6 +1577,12 @@ private:
 	 */
 	PAKFILE_API bool ShouldValidatePrunedDirectory() const;
 
+	/**
+	 * Returns whether the current PakFile is using a TDirectoryTree for storage of its DirectoryTreeIndex rather than a TMap.
+	 * Returns false if !ENABLE_PAKFILE_USE_DIRECTORY_TREE or if PakFile TDirectoryTrees are disabled.
+	 */
+	PAKFILE_API bool ShouldUseDirectoryTree() const;
+
 	/**
 	  * Add the given (Filename,FPakEntryLocation) value into the provided indexes
 	  *
@@ -1652,30 +1591,45 @@ private:
 	  * @param MountPoint The mount point of the pakfile containing the FPakEntryLocation, used to create the filename in the DirectoryIndex
 	  * @param PathHashSeed the pakfile-specific seed for the hash of the path in the PathHasIndex
 	  * @param DirectoryIndex Optional FDirectoryIndex into which to insert the (Filename, FPakEntryLocation)
+	  * @param DirectoryTreeIndex Optional FDirectoryTreeIndex into which to insert the (Filename, FPakEntryLocation)
 	  * @param PathHashIndex Optional FPathHashIndex into which to insert the (Filename, FPakEntryLocation)
 	  * @param InOutCollisionDetection Optional parameter to detect hash collisions.  If present, the hashed filename will be check()'d for a collision against a different filename in InOutCollisionDetection, and will be added into InOutCollisionDetection
 	  * @param PakFileVersion Version of the pakfile containing the index, to support legacy formats
 	  */
-	static PAKFILE_API void AddEntryToIndex(const FString& Filename, const FPakEntryLocation& EntryLocation, const FString& MountPoint, uint64 PathHashSeed,
-		FDirectoryIndex* DirectoryIndex, FPathHashIndex* PathHashIndex, TMap<uint64, FString>* CollisionDetection, int32 PakFileVersion);
+	static void AddEntryToIndex(const FString& Filename, const FPakEntryLocation& EntryLocation,
+		const FString& MountPoint, uint64 PathHashSeed, FDirectoryIndex* DirectoryIndex,
+		FDirectoryTreeIndex* DirectoryTreeIndex, FPathHashIndex* PathHashIndex,
+		TMap<uint64, FString>* CollisionDetection, int32 PakFileVersion);
 
 	/* Encodes a pak entry as an array of bytes into the given archive.  Returns true if encoding succeeded.  If encoding did not succeed, caller will need to store the InPakEntry in an unencoded list */
-	static PAKFILE_API bool EncodePakEntry(FArchive& Ar, const FPakEntry& InPakEntry, const FPakInfo& InInfo);
+	static bool EncodePakEntry(FArchive& Ar, const FPakEntry& InPakEntry, const FPakInfo& InInfo);
 
 	/* Decodes a bit-encoded pak entry from a pointer to the start of its encoded bytes into the given OutEntry */
-	static PAKFILE_API void DecodePakEntry(const uint8* SourcePtr, FPakEntry& OutEntry, const FPakInfo& InInfo);
+	static void DecodePakEntry(const uint8* SourcePtr, FPakEntry& OutEntry, const FPakInfo& InInfo);
 
 	/* Internal index loading function that returns false if index loading fails due to an intermittent IO error. Allows LoadIndex to retry or throw a fatal as required */
-	PAKFILE_API bool LoadIndexInternal(FArchive& Reader);
+	bool LoadIndexInternal(FArchive& Reader, FDirectoryIndex& OutDirectoryTMap,
+		FDirectoryTreeIndex& OutDirectoryTree, FDirectoryIndex& OutPrunedDirectoryTMap,
+		FDirectoryTreeIndex& OutPrunedDirectoryTree, bool bStoreDirectoryTMap, bool bStoreDirectoryTree);
+	void LoadIndexInternal_DirectoryIndex(FArchive& Ar, FDirectoryIndex& OutDirectoryTMap,
+		FDirectoryTreeIndex& OutDirectoryTree, bool bLoadIntoDirectoryTMap, bool bLoadIntoDirectoryTree);
+	// This is an internal function that should only be called by FPakFile or by PakFileUtilities.cpp; it can change without deprecation.
+public:
+	PAKFILE_API static void SaveIndexInternal_DirectoryIndex(FArchive& Ar, const FDirectoryIndex& DirectoryTMap);
+private:
 
 	/* Legacy index loading function for PakFiles saved before FPakInfo::PakFile_Version_PathHashIndex */
-	PAKFILE_API bool LoadLegacyIndex(FArchive& Reader);
+	bool LoadLegacyIndex(FArchive& Reader, FDirectoryIndex& OutDirectoryTMap);
 
 	/* Helper function for LoadIndexInternal; each array of Index bytes read from the file needs to be independently decrypted and checked for corruption */
-	PAKFILE_API bool DecryptAndValidateIndex(FArchive& Reader, TArray<uint8>& IndexData, FSHAHash& InExpectedHash, FSHAHash& OutActualHash);
+	bool DecryptAndValidateIndex(FArchive& Reader, TArray<uint8>& IndexData, FSHAHash& InExpectedHash, FSHAHash& OutActualHash);
 
 	/* Manually add a file to a pak file */
-	PAKFILE_API void AddSpecialFile(const FPakEntry& Entry, const FString& Filename);
+	void AddSpecialFile(const FPakEntry& Entry, const FString& Filename);
+
+	static void PruneDirectoryIndexInternal(FDirectoryIndex* InOutDirectoryIndex,
+		FDirectoryTreeIndex* InOutDirectoryTreeIndex, FDirectoryIndex* PrunedDirectoryIndex,
+		FDirectoryTreeIndex* PrunedDirectoryTreeIndex, const FString& MountPoint);
 
 	template <typename ShouldVisitFunc>
 	struct FVisitFilter
@@ -1690,19 +1644,28 @@ private:
 	template <typename ShouldVisitFunc, class ContainerType>
 	void FindPrunedFilesAtPathInternal(const TCHAR* InPath, ContainerType& OutFiles, const FVisitFilter<ShouldVisitFunc>& VisitFilter) const;
 
+	/**
+	 * Lookup the FPakEntryLocation stored in the given DirectoryIndex or DirectoryTreeIndex,
+	 * return nullptr if not found
+	 */
+	PAKFILE_API const FPakEntryLocation* FindLocationFromIndex(const FString& FullPath,
+		const FDirectoryIndex& InDirectoryIndex, const FDirectoryTreeIndex& InDirectoryTreeIndex) const;
+
 	/**
 	 * Search the given FDirectoryIndex for all files under the given Directory.  Helper for FindFilesAtPath, called separately on the DirectoryIndex or Pruned DirectoryIndex. Does not use
 	 * FScopedPakDirectoryIndexAccess internally; caller is responsible for calling from within a lock.
 	 * Returned paths are full paths (include the mount point)
 	 */
 	template <typename ShouldVisitFunc, class ContainerType>
-	void FindFilesAtPathInIndex(const FDirectoryIndex& TargetIndex, const FDirectoryTreeIndex& TargetTreeIndex, ContainerType& OutFiles, const FString& Directory, 
+	void FindFilesAtPathInIndex(const FDirectoryIndex& TargetIndex, const FDirectoryTreeIndex& TargetTreeIndex,
+		ContainerType& OutFiles, const FString& FullSearchPath,
 		const FVisitFilter<ShouldVisitFunc>& VisitFilter) const;
 
 #if ENABLE_PAKFILE_USE_DIRECTORY_TREE
 	template <typename ShouldVisitFunc, class ContainerType>
-	void FindFilesAtPathInTreeIndexInternal(const FString& Path, const FDirectoryTreeIndex& TargetTreeIndex, ContainerType& OutFiles, TArray<FString>& OutDirectories,
-		const FString& Directory, const FVisitFilter<ShouldVisitFunc>& VisitFilter) const;
+	void FindFilesAtPathInTreeIndexInternal(FStringView RelSearchPath, const FDirectoryTreeIndex& TargetTreeIndex,
+		ContainerType& OutFiles, TArray<FString>& OutDirectories,
+		const FString& FullSearchPath, const FVisitFilter<ShouldVisitFunc>& VisitFilter) const;
 
 #if !UE_BUILD_SHIPPING
 	template <class ContainerType>
@@ -1711,12 +1674,14 @@ private:
 #endif	// ENABLE_PAKFILE_USE_DIRECTORY_TREE
 
 	template <typename ShouldVisitFunc, class ContainerType>
-	void FindFilesAtPathInIndexInternal(const FStringView& Path, const FDirectoryIndex& TargetIndex, ContainerType& OutFiles, TArray<FString>& OutDirectories,
-		const FString& Directory, const FVisitFilter<ShouldVisitFunc>& VisitFilter) const;
+	void FindFilesAtPathInIndexInternal(const FStringView& RelSearchPath, const FDirectoryIndex& TargetIndex,
+		ContainerType& OutFiles, TArray<FString>& OutDirectories,
+		const FString& FullSearchPath, const FVisitFilter<ShouldVisitFunc>& VisitFilter) const;
 
 	template <typename ShouldVisitFunc, class ContainerType>
-	static void FindFilesAtPathInPakDirectoryInternal(const FString& MountPoint, const FString& Path, const FPakDirectory& PakDirectory, ContainerType& OutFiles, TArray<FString>& OutDirectoriesInPak,
-		const FString& Directory, const FVisitFilter<ShouldVisitFunc>& VisitFilter);
+	static void FindFilesAtPathInPakDirectoryInternal(const FString& MountPoint, FStringView RelPathInIndex,
+		const FPakDirectory& PakDirectory, ContainerType& OutFiles, TArray<FString>& OutDirectoriesInPak,
+		const FString& FullSearchPath, const FVisitFilter<ShouldVisitFunc>& VisitFilter);
 
 	/** Converts the path to a RelativePathFromMount and normalizes it to the expected format for Pak Directories.  Returns false if Path is not under the MountDir and hence can not be in this PakFile. */
 	PAKFILE_API bool NormalizeDirectoryQuery(const TCHAR* InPath, FString& OutRelativePathFromMount) const;
@@ -1726,6 +1691,8 @@ private:
 	 * Validates the result if IsPakValidatePruning. Does not use a critical section; caller is responsible for calling from within an FScopedPakDirectoryIndexAccess
 	 */
 	PAKFILE_API const FPakDirectory* FindPrunedDirectoryInternal(const FString& RelativePathFromMount) const;
+	const FPakDirectory* FindPrunedDirectoryInIndexInternal(const FString& RelativePathFromMount,
+		const FDirectoryIndex& InDirectoryIndex, const FDirectoryTreeIndex& InTreeIndex) const;
 
 #if ENABLE_PAKFILE_RUNTIME_PRUNING_VALIDATE
 	/* Logs an error if the two sets are not identical after removing all config-specified ignore paths */
@@ -2478,12 +2445,12 @@ public:
 			const FPakDirectory* PakDirectory = PakFile->FindPrunedDirectory(*Path);
 			if (PakDirectory != nullptr)
 			{
-				for (FPakDirectory::TConstIterator DirectoryIt(*PakDirectory); DirectoryIt; ++DirectoryIt)
+				for (FPakDirectory::TConstIterator FileIt(*PakDirectory); FileIt; ++FileIt)
 				{
 					FPakEntry PakEntry;
-					if (PakFile->GetPakEntry(DirectoryIt.Value(), &PakEntry) != FPakFile::EFindResult::NotFound && PakEntry.Offset == FileEntry.Offset)
+					if (PakFile->GetPakEntry(FileIt.Value(), &PakEntry) != FPakFile::EFindResult::NotFound && PakEntry.Offset == FileEntry.Offset)
 					{
-						const FString& RealFilename = DirectoryIt.Key();
+						const FString& RealFilename = FileIt.Key();
 						return Path / RealFilename;
 					}
 				}
