PakFiles: Enable by default the new DirectoryTree index for pakfiles, to improve performance and memory of pakfile queries.
#rnx

diff --git a/Engine/Source/Runtime/PakFile/Public/IPlatformFilePak.h b/Engine/Source/Runtime/PakFile/Public/IPlatformFilePak.h
--- a/Engine/Source/Runtime/PakFile/Public/IPlatformFilePak.h
+++ b/Engine/Source/Runtime/PakFile/Public/IPlatformFilePak.h
@@ -52,7 +52,7 @@
 #include "UObject/UnrealNames.h"
 
 #ifndef ENABLE_PAKFILE_USE_DIRECTORY_TREE
-#define ENABLE_PAKFILE_USE_DIRECTORY_TREE 0
+#define ENABLE_PAKFILE_USE_DIRECTORY_TREE 1
 #endif
 
 class FChunkCacheWorker;
