Introduced interface for I/O store file backend in order to toggle between old and new implementation.

diff --git a/Engine/Source/Runtime/PakFile/Private/FileIoDispatcherBackend.h b/Engine/Source/Runtime/PakFile/Private/FileIoDispatcherBackend.h
new file mode 100644
--- /dev/null
+++ b/Engine/Source/Runtime/PakFile/Private/FileIoDispatcherBackend.h
@@ -0,0 +1,25 @@
+// Copyright Epic Games, Inc. All Rights Reserved.
+
+#pragma once
+
+#include "IO/IoDispatcherBackend.h"
+#include "IO/IoStatus.h"
+#include "Misc/AES.h"
+
+struct FIoContainerHeader;
+
+namespace UE::IoStore
+{
+
+class IFileIoDispatcherBackend
+	: public IIoDispatcherBackend
+{
+public:
+	virtual ~IFileIoDispatcherBackend() = default;
+
+	virtual TIoStatusOr<FIoContainerHeader> Mount(const TCHAR* TocPath, int32 Order, const FGuid& EncryptionKeyGuid, const FAES::FAESKey& EncryptionKey) = 0;
+	virtual bool							Unmount(const TCHAR* TocPath) = 0;
+	virtual void							ReopenAllFileHandles() = 0;
+};
+
+} // namespace UE::IoStore
diff --git a/Engine/Source/Runtime/PakFile/Private/IoDispatcherFileBackend.h b/Engine/Source/Runtime/PakFile/Private/IoDispatcherFileBackend.h
--- a/Engine/Source/Runtime/PakFile/Private/IoDispatcherFileBackend.h
+++ b/Engine/Source/Runtime/PakFile/Private/IoDispatcherFileBackend.h
@@ -2,6 +2,8 @@
 
 #pragma once
 
+#include "FileIoDispatcherBackend.h"
+
 #include "IoDispatcherFileBackendTypes.h"
 #include "IO/IoDispatcher.h"
 #include "IO/IoStore.h"
@@ -103,15 +105,15 @@ private:
 
 class FFileIoStore final
 	: public FRunnable
-	, public IIoDispatcherBackend
+	, public UE::IoStore::IFileIoDispatcherBackend
 {
 public:
 	FFileIoStore(TUniquePtr<IPlatformFileIoStore>&& PlatformImpl);
 	~FFileIoStore();
 	void Initialize(TSharedRef<const FIoDispatcherBackendContext> Context) override;
 	void Shutdown() override;
-	TIoStatusOr<FIoContainerHeader> Mount(const TCHAR* InTocPath, int32 Order, const FGuid& EncryptionKeyGuid, const FAES::FAESKey& EncryptionKey);
-	bool Unmount(const TCHAR* InTocPath);
+	virtual TIoStatusOr<FIoContainerHeader> Mount(const TCHAR* InTocPath, int32 Order, const FGuid& EncryptionKeyGuid, const FAES::FAESKey& EncryptionKey) override;
+	virtual bool Unmount(const TCHAR* InTocPath) override;
 	void ResolveIoRequests(FIoRequestList Requests, FIoRequestList& OutUnresolved) override;
 	void CancelIoRequest(FIoRequestImpl* Request) override;
 	void UpdatePriorityForIoRequest(FIoRequestImpl* Request) override;
@@ -120,7 +122,7 @@ public:
 	FIoRequestImpl* GetCompletedIoRequests() override;
 	TIoStatusOr<FIoMappedRegion> OpenMapped(const FIoChunkId& ChunkId, const FIoReadOptions& Options) override;
 	const TCHAR* GetName() const override;
-	void ReopenAllFileHandles();
+	virtual void ReopenAllFileHandles() override;
 
 	virtual bool Init() override;
 	virtual uint32 Run() override;
