Extract I/O dispatcher filesystem stats to separate file.

diff --git a/Engine/Source/Runtime/Core/Internal/IO/IoDispatcherFilesystemStats.cpp b/Engine/Source/Runtime/Core/Internal/IO/IoDispatcherFilesystemStats.cpp
new file mode 100644
--- /dev/null
+++ b/Engine/Source/Runtime/Core/Internal/IO/IoDispatcherFilesystemStats.cpp
@@ -0,0 +1,316 @@
+// Copyright Epic Games, Inc. All Rights Reserved.
+
+#include "IO/IoDispatcherFilesystemStats.h"
+#include "ProfilingDebugging/CsvProfiler.h"
+
+namespace UE
+{
+
+#if UE_IO_DISPATCHER_FILESYSTEM_STATS_ENABLED 
+
+////////////////////////////////////////////////////////////////////////////////
+CSV_DEFINE_CATEGORY(IoDispatcherFileBackend, true);
+CSV_DEFINE_CATEGORY(IoDispatcherFileBackendVerbose, false);
+
+// These stats go to both insights and csv by default
+// TODO: Ideally these should go to insights even if CSV is not capturing, but not be doubled-up if both CSV and Insights are capturing
+// TODO: It would also be nice to send these to insights as int64 without unit conversion where appropriate
+// IoDispatcher thread
+CSV_DEFINE_STAT(IoDispatcherFileBackend,			FrameBytesScatteredKB);
+CSV_DEFINE_STAT(IoDispatcherFileBackend,			QueuedFilesystemReadMB);
+CSV_DEFINE_STAT(IoDispatcherFileBackendVerbose,	QueuedFilesystemReads);
+CSV_DEFINE_STAT(IoDispatcherFileBackendVerbose,	QueuedUncompressBlocks);	
+CSV_DEFINE_STAT(IoDispatcherFileBackendVerbose,	QueuedUncompressInMB);		
+CSV_DEFINE_STAT(IoDispatcherFileBackendVerbose,	QueuedUncompressOutMB);	
+CSV_DEFINE_STAT(IoDispatcherFileBackendVerbose,	FrameBytesReadKB);				
+CSV_DEFINE_STAT(IoDispatcherFileBackendVerbose,	FrameBytesUncompressedInKB);	
+CSV_DEFINE_STAT(IoDispatcherFileBackendVerbose,	FrameBytesUncompressedOutKB);	
+
+// FileIoStore thread			
+CSV_DEFINE_STAT(IoDispatcherFileBackend,			FrameFilesystemBytesReadKB);	
+CSV_DEFINE_STAT(IoDispatcherFileBackend,			FrameSequentialReads);
+CSV_DEFINE_STAT(IoDispatcherFileBackend,			FrameSeeks);
+CSV_DEFINE_STAT(IoDispatcherFileBackendVerbose,	FrameFilesystemReads);
+CSV_DEFINE_STAT(IoDispatcherFileBackendVerbose,	FrameForwardSeeks);
+CSV_DEFINE_STAT(IoDispatcherFileBackendVerbose,	FrameBackwardSeeks);
+CSV_DEFINE_STAT(IoDispatcherFileBackendVerbose,	FrameHandleChangeSeeks);
+CSV_DEFINE_STAT(IoDispatcherFileBackendVerbose,	FrameSeekDistanceMB);	
+CSV_DEFINE_STAT(IoDispatcherFileBackendVerbose,	FrameBlockCacheStores);
+CSV_DEFINE_STAT(IoDispatcherFileBackendVerbose,	FrameBlockCacheStoresKB);
+CSV_DEFINE_STAT(IoDispatcherFileBackendVerbose,	FrameBlockCacheHits);	
+CSV_DEFINE_STAT(IoDispatcherFileBackendVerbose,	FrameBlockCacheHitKB);	
+CSV_DEFINE_STAT(IoDispatcherFileBackendVerbose,	FrameBlockCacheMisses);
+CSV_DEFINE_STAT(IoDispatcherFileBackendVerbose,	FrameBlockCacheMissKB);
+
+////////////////////////////////////////////////////////////////////////////////
+namespace IoDispatcherFilesystemStats
+{
+
+static float BytesToApproxMB(uint64 Bytes)
+{
+	return float(double(Bytes) / 1024.0 / 1024.0);
+}
+
+static float BytesToApproxKB(uint64 Bytes)
+{
+	return float(double(Bytes) / 1024.0);
+}
+
+} // namespace IoDispatcherFilesystemStats
+
+////////////////////////////////////////////////////////////////////////////////
+FIoDispatcherFilesystemStats::FIoDispatcherFilesystemStats()
+#if COUNTERSTRACE_ENABLED
+	: QueuedReadRequestsSizeCounter(TEXT("FileIoStore/QueuedReadRequestsSize"), TraceCounterDisplayHint_Memory)
+	, CompletedReadRequestsSizeCounter(TEXT("FileIoStore/CompletedReadRequestsSize"), TraceCounterDisplayHint_Memory)
+	, QueuedCompressedSizeCounter(TEXT("FileIoStore/QueuedCompressedSize"), TraceCounterDisplayHint_Memory)
+	, QueuedUncompressedSizeCounter(TEXT("FileIoStore/QueuedUncompressedSize"), TraceCounterDisplayHint_Memory)
+	, CompletedCompressedSizeCounter(TEXT("FileIoStore/CompletedCompressedSize"), TraceCounterDisplayHint_Memory)
+	, CompletedUncompressedSizeCounter(TEXT("FileIoStore/CompletedUncompressedSize"), TraceCounterDisplayHint_Memory)
+	, FileSystemSeeksTotalDistanceCounter(TEXT("FileIoStore/FileSystemSeeksTotalDistance"), TraceCounterDisplayHint_Memory)
+	, FileSystemSeeksForwardCountCounter(TEXT("FileIoStore/FileSystemSeeksForwardCount"), TraceCounterDisplayHint_None)
+	, FileSystemSeeksBackwardCountCounter(TEXT("FileIoStore/FileSystemSeeksBackwardCount"), TraceCounterDisplayHint_None)
+	, FileSystemSeeksChangeHandleCountCounter(TEXT("FileIoStore/FileSystemSeeksChangeHandleCount"), TraceCounterDisplayHint_None)
+	, FileSystemCompletedRequestsSizeCounter(TEXT("FileIoStore/FileSystemCompletedRequestsSize"), TraceCounterDisplayHint_Memory)
+	, BlockCacheStoredSizeCounter(TEXT("FileIoStore/BlockCacheStoredSize"), TraceCounterDisplayHint_Memory)
+	, BlockCacheHitSizeCounter(TEXT("FileIoStore/BlockCacheHitSize"), TraceCounterDisplayHint_Memory)
+	, BlockCacheMissedSizeCounter(TEXT("FileIoStore/BlockCacheMissedSize"), TraceCounterDisplayHint_Memory)
+	, ScatteredSizeCounter(TEXT("FileIoStore/ScatteredSize"), TraceCounterDisplayHint_Memory)
+	, TocMemoryCounter(TEXT("FileIoStore/TocMemory"), TraceCounterDisplayHint_Memory)
+	, AvailableBuffersCounter(TEXT("FileIoStore/AvailableBuffers"), TraceCounterDisplayHint_None)
+#endif
+{
+#if CSV_PROFILER_STATS
+	TickerHandle = FTSTicker::GetCoreTicker().AddTicker(FTickerDelegate::CreateRaw(this, &FIoDispatcherFilesystemStats::CsvTick));
+#endif
+}
+
+FIoDispatcherFilesystemStats::~FIoDispatcherFilesystemStats()
+{
+#if CSV_PROFILER_STATS
+	FTSTicker::GetCoreTicker().RemoveTicker(TickerHandle);
+#endif
+}
+
+bool FIoDispatcherFilesystemStats::CsvTick(float DeltaTime)
+{
+	using namespace IoDispatcherFilesystemStats;
+#if CSV_PROFILER_STATS
+	CSV_CUSTOM_STAT_DEFINED(QueuedFilesystemReadMB, BytesToApproxMB(QueuedFilesystemReadBytes), ECsvCustomStatOp::Set);
+	CSV_CUSTOM_STAT_DEFINED(QueuedFilesystemReads, (int32)QueuedFilesystemReads, ECsvCustomStatOp::Set);
+
+	CSV_CUSTOM_STAT_DEFINED(QueuedUncompressBlocks, int32(QueuedUncompressBlocks), ECsvCustomStatOp::Set);
+	CSV_CUSTOM_STAT_DEFINED(QueuedUncompressInMB, BytesToApproxMB(QueuedUncompressBytesIn), ECsvCustomStatOp::Set);
+	CSV_CUSTOM_STAT_DEFINED(QueuedUncompressOutMB, BytesToApproxMB(QueuedUncompressBytesOut), ECsvCustomStatOp::Set);
+#endif
+
+	return true;
+}
+
+void FIoDispatcherFilesystemStats::OnReadRequestsQueued(uint64 ByteCount, uint64 ReadCount)
+{
+#if CSV_PROFILER_STATS
+	QueuedFilesystemReadBytes += ByteCount;
+	QueuedFilesystemReads += ReadCount;
+#endif
+#if COUNTERSTRACE_ENABLED
+	QueuedReadRequestsSizeCounter.Add(ByteCount);
+#endif
+}
+
+void FIoDispatcherFilesystemStats::OnFilesystemReadStarted(uint64 FileHandle, uint64 Offset, uint64 Size)
+{
+	CSV_CUSTOM_STAT_DEFINED(FrameFilesystemReads, 1, ECsvCustomStatOp::Accumulate);
+
+	if (LastFileReadInfo.FileHandle != FileHandle) 
+	{
+		OnHandleChangeSeek();
+	}
+	else if (LastFileReadInfo.Offset == Offset)
+	{
+		OnSequentialRead();
+	}
+	else
+	{
+		OnSeek(LastFileReadInfo.Offset, Offset); 
+	}
+
+	LastFileReadInfo = FFileReadInfo
+	{
+		.FileHandle = FileHandle,
+		.Offset = Offset + Size
+	};
+}
+
+void FIoDispatcherFilesystemStats::OnFilesystemReadCompleted(uint64 FileHandle, uint64 Offset, uint64 Size)
+{
+	CSV_CUSTOM_STAT_DEFINED(FrameFilesystemBytesReadKB, IoDispatcherFilesystemStats::BytesToApproxKB(Size), ECsvCustomStatOp::Accumulate);
+#if COUNTERSTRACE_ENABLED
+	FileSystemCompletedRequestsSizeCounter.Add(Size);
+#endif
+}
+
+void FIoDispatcherFilesystemStats::OnReadRequestsCompleted(uint64 ByteCount, uint64 ReadCount)
+{
+#if CSV_PROFILER_STATS
+	QueuedFilesystemReadBytes -= ByteCount;
+	QueuedFilesystemReads -= ReadCount;
+#endif
+
+	CSV_CUSTOM_STAT_DEFINED(FrameBytesReadKB, IoDispatcherFilesystemStats::BytesToApproxKB(ByteCount), ECsvCustomStatOp::Accumulate); // TODO: Send to insights if CSV is disabled
+
+#if COUNTERSTRACE_ENABLED
+	QueuedReadRequestsSizeCounter.Subtract(ByteCount);
+	CompletedReadRequestsSizeCounter.Add(ByteCount);
+#endif
+}
+
+void FIoDispatcherFilesystemStats::OnDecompressQueued(uint64 CompressedSize, uint64 UncompressedSize)
+{
+#if CSV_PROFILER_STATS
+	++QueuedUncompressBlocks;
+	QueuedUncompressBytesIn += CompressedSize;
+	QueuedUncompressBytesOut += UncompressedSize;
+#endif
+
+#if COUNTERSTRACE_ENABLED
+	QueuedCompressedSizeCounter.Add(CompressedSize);
+	QueuedUncompressedSizeCounter.Add(UncompressedSize);
+#endif
+}
+
+void FIoDispatcherFilesystemStats::OnDecompressComplete(uint64 CompressedSize, uint64 UncompressedSize)
+{
+	using namespace IoDispatcherFilesystemStats;
+#if CSV_PROFILER_STATS
+	--QueuedUncompressBlocks;
+	QueuedUncompressBytesIn -= CompressedSize;
+	QueuedUncompressBytesOut -= UncompressedSize;
+#endif
+
+	CSV_CUSTOM_STAT_DEFINED(FrameBytesUncompressedInKB, BytesToApproxKB(CompressedSize), ECsvCustomStatOp::Accumulate); 
+	CSV_CUSTOM_STAT_DEFINED(FrameBytesUncompressedOutKB, BytesToApproxKB(UncompressedSize), ECsvCustomStatOp::Accumulate);
+
+#if COUNTERSTRACE_ENABLED
+	QueuedCompressedSizeCounter.Subtract(CompressedSize);
+	QueuedUncompressedSizeCounter.Subtract(UncompressedSize);
+
+	CompletedCompressedSizeCounter.Add(CompressedSize);
+	CompletedUncompressedSizeCounter.Add(UncompressedSize);
+#endif
+}
+
+void FIoDispatcherFilesystemStats::OnBytesScattered(int64 NumBytes)
+{
+	CSV_CUSTOM_STAT_DEFINED(FrameBytesScatteredKB, IoDispatcherFilesystemStats::BytesToApproxKB(NumBytes), ECsvCustomStatOp::Accumulate);
+#if COUNTERSTRACE_ENABLED
+	ScatteredSizeCounter.Add(NumBytes);
+#endif
+}
+
+void FIoDispatcherFilesystemStats::OnSequentialRead()
+{
+	CSV_CUSTOM_STAT_DEFINED(FrameSequentialReads, 1, ECsvCustomStatOp::Accumulate);
+}
+
+void FIoDispatcherFilesystemStats::OnSeek(uint64 PrevOffset, uint64 NewOffset)
+{
+	using namespace IoDispatcherFilesystemStats;
+	if (NewOffset > PrevOffset)
+	{
+		int64 Delta = NewOffset - PrevOffset;
+
+		CSV_CUSTOM_STAT_DEFINED(FrameForwardSeeks, 1, ECsvCustomStatOp::Accumulate);
+		CSV_CUSTOM_STAT_DEFINED(FrameSeekDistanceMB, BytesToApproxMB(Delta), ECsvCustomStatOp::Accumulate);
+
+#if COUNTERSTRACE_ENABLED
+		FileSystemSeeksTotalDistanceCounter.Add(Delta);
+		FileSystemSeeksForwardCountCounter.Increment();
+#endif
+	}
+	else
+	{
+		int64 Delta = PrevOffset - NewOffset;
+		CSV_CUSTOM_STAT_DEFINED(FrameBackwardSeeks, 1, ECsvCustomStatOp::Accumulate);
+		CSV_CUSTOM_STAT_DEFINED(FrameSeekDistanceMB, BytesToApproxMB(Delta), ECsvCustomStatOp::Accumulate);
+
+#if COUNTERSTRACE_ENABLED
+		FileSystemSeeksTotalDistanceCounter.Add(Delta);
+		FileSystemSeeksBackwardCountCounter.Increment();
+#endif
+	}
+
+	CSV_CUSTOM_STAT_DEFINED(FrameSeeks, 1, ECsvCustomStatOp::Accumulate);
+}
+
+void FIoDispatcherFilesystemStats::OnHandleChangeSeek()
+{
+	CSV_CUSTOM_STAT_DEFINED(FrameHandleChangeSeeks, 1, ECsvCustomStatOp::Accumulate);
+	CSV_CUSTOM_STAT_DEFINED(FrameSeeks, 1, ECsvCustomStatOp::Accumulate);
+
+#if COUNTERSTRACE_ENABLED
+	FileSystemSeeksChangeHandleCountCounter.Increment();
+#endif
+}
+
+void FIoDispatcherFilesystemStats::OnBlockCacheStore(uint64 NumBytes)
+{
+	CSV_CUSTOM_STAT_DEFINED(FrameBlockCacheStores, 1, ECsvCustomStatOp::Accumulate);
+	CSV_CUSTOM_STAT_DEFINED(FrameBlockCacheStoresKB, IoDispatcherFilesystemStats::BytesToApproxKB(NumBytes), ECsvCustomStatOp::Accumulate);
+
+#if COUNTERSTRACE_ENABLED
+	BlockCacheStoredSizeCounter.Add(NumBytes);
+#endif
+}
+
+void FIoDispatcherFilesystemStats::OnBlockCacheHit(uint64 NumBytes)
+{
+	CSV_CUSTOM_STAT_DEFINED(FrameBlockCacheHits, 1, ECsvCustomStatOp::Accumulate);
+	CSV_CUSTOM_STAT_DEFINED(FrameBlockCacheHitKB, IoDispatcherFilesystemStats::BytesToApproxKB(NumBytes), ECsvCustomStatOp::Accumulate);
+
+#if COUNTERSTRACE_ENABLED
+	BlockCacheHitSizeCounter.Add(NumBytes);
+#endif
+}
+
+void FIoDispatcherFilesystemStats::OnBlockCacheMiss(uint64 NumBytes)
+{
+	CSV_CUSTOM_STAT_DEFINED(FrameBlockCacheMisses, 1, ECsvCustomStatOp::Accumulate);
+	CSV_CUSTOM_STAT_DEFINED(FrameBlockCacheMissKB, IoDispatcherFilesystemStats::BytesToApproxKB(NumBytes), ECsvCustomStatOp::Accumulate);
+
+#if COUNTERSTRACE_ENABLED
+	BlockCacheMissedSizeCounter.Add(NumBytes);
+#endif
+}
+
+void FIoDispatcherFilesystemStats::OnTocMounted(uint64 AllocatedSize)
+{
+#if COUNTERSTRACE_ENABLED
+	TocMemoryCounter.Add(AllocatedSize);
+#endif
+}
+
+void FIoDispatcherFilesystemStats::OnTocUnmounted(uint64 AllocatedSize)
+{
+#if COUNTERSTRACE_ENABLED
+	TocMemoryCounter.Subtract(AllocatedSize);
+#endif
+}
+
+void FIoDispatcherFilesystemStats::OnBufferReleased()
+{
+#if COUNTERSTRACE_ENABLED
+	AvailableBuffersCounter.Increment();
+#endif
+}
+
+void FIoDispatcherFilesystemStats::OnBufferAllocated()
+{
+#if COUNTERSTRACE_ENABLED
+	AvailableBuffersCounter.Decrement();
+#endif
+}
+
+#endif // UE_IO_DISPATCHER_FILESYSTEM_STATS_ENABLED 
+
+} // namespace UE
diff --git a/Engine/Source/Runtime/Core/Internal/IO/IoDispatcherFilesystemStats.h b/Engine/Source/Runtime/Core/Internal/IO/IoDispatcherFilesystemStats.h
new file mode 100644
--- /dev/null
+++ b/Engine/Source/Runtime/Core/Internal/IO/IoDispatcherFilesystemStats.h
@@ -0,0 +1,104 @@
+// Copyright Epic Games, Inc. All Rights Reserved.
+
+#pragma once
+
+#include "Containers/Ticker.h"
+#include "Math/NumericLimits.h"
+#include "ProfilingDebugging/CountersTrace.h"
+
+#define UE_IO_DISPATCHER_FILESYSTEM_STATS_ENABLED (COUNTERSTRACE_ENABLED || CSV_PROFILER_STATS)
+
+namespace UE
+{
+
+#if UE_IO_DISPATCHER_FILESYSTEM_STATS_ENABLED 
+class FIoDispatcherFilesystemStats
+{
+	struct FFileReadInfo
+	{
+		uint64	FileHandle	= 0;
+		uint64	Offset		= MAX_uint64;
+	};
+
+public:
+	CORE_API		FIoDispatcherFilesystemStats();
+	CORE_API		~FIoDispatcherFilesystemStats();
+	CORE_API bool	CsvTick(float DeltaTime);
+	CORE_API void	OnReadRequestsQueued(uint64 ByteCount, uint64 ReadCount);
+	CORE_API void	OnFilesystemReadStarted(uint64 FileHandle, uint64 Offset, uint64 Size);
+	CORE_API void 	OnFilesystemReadCompleted(uint64 FileHandle, uint64 Offset, uint64 Size);
+	CORE_API void 	OnReadRequestsCompleted(uint64 ByteCount, uint64 ReadCount);
+	CORE_API void 	OnDecompressQueued(uint64 CompressedSize, uint64 UncompressedSize);
+	CORE_API void 	OnDecompressComplete(uint64 CompressedSize, uint64 UncompressedSize);
+	CORE_API void	OnBytesScattered(int64 BytesScattered);
+	CORE_API void 	OnBlockCacheStore(uint64 NumBytes);
+	CORE_API void 	OnBlockCacheHit(uint64 NumBytes);
+	CORE_API void 	OnBlockCacheMiss(uint64 NumBytes);
+	CORE_API void 	OnTocMounted(uint64 AllocatedSize);
+	CORE_API void 	OnTocUnmounted(uint64 AllocatedSize);
+	CORE_API void 	OnBufferReleased();
+	CORE_API void 	OnBufferAllocated();
+
+private:
+	void			OnSequentialRead();
+	void			OnSeek(uint64 PrevOffset, uint64 NewOffset);
+	void			OnHandleChangeSeek();
+
+#if COUNTERSTRACE_ENABLED
+	FCountersTrace::FCounterInt QueuedReadRequestsSizeCounter;
+	FCountersTrace::FCounterInt CompletedReadRequestsSizeCounter;
+	FCountersTrace::FCounterInt QueuedCompressedSizeCounter;
+	FCountersTrace::FCounterInt QueuedUncompressedSizeCounter;
+	FCountersTrace::FCounterInt CompletedCompressedSizeCounter;
+	FCountersTrace::FCounterInt CompletedUncompressedSizeCounter;
+	FCountersTrace::FCounterInt FileSystemSeeksTotalDistanceCounter;
+	FCountersTrace::FCounterInt FileSystemSeeksForwardCountCounter;
+	FCountersTrace::FCounterInt FileSystemSeeksBackwardCountCounter;
+	FCountersTrace::FCounterInt FileSystemSeeksChangeHandleCountCounter;
+	FCountersTrace::FCounterInt FileSystemCompletedRequestsSizeCounter;
+	FCountersTrace::FCounterInt BlockCacheStoredSizeCounter;
+	FCountersTrace::FCounterInt BlockCacheHitSizeCounter;
+	FCountersTrace::FCounterInt BlockCacheMissedSizeCounter;
+	FCountersTrace::FCounterInt ScatteredSizeCounter;
+	FCountersTrace::FCounterInt TocMemoryCounter;
+	FCountersTrace::TCounter<std::atomic<int64>, TraceCounterType_Int> AvailableBuffersCounter;
+#endif
+#if CSV_PROFILER_STATS
+	uint64						QueuedFilesystemReadBytes = 0;
+	uint64 						QueuedFilesystemReads = 0;
+	uint64 						QueuedUncompressBytesIn = 0;
+	uint64 						QueuedUncompressBytesOut = 0;
+	uint64 						QueuedUncompressBlocks = 0;
+#endif
+	FTSTicker::FDelegateHandle	TickerHandle;
+	FFileReadInfo				LastFileReadInfo;
+};
+
+#else
+class FIoDispatcherFilesystemStats
+{
+public:
+			FIoDispatcherFilesystemStats() {}
+			~FIoDispatcherFilesystemStats() {}
+	void	OnReadRequestsQueued(uint64 ByteCount, uint64 ReadCount) {};
+	void	OnFilesystemReadStarted(uint64 FileHandle, uint64 Offset, uint64 Size) {};
+	void 	OnFilesystemReadCompleted(uint64 FileHandle, uint64 Offset, uint64 Size) {};
+	void 	OnReadRequestsCompleted(uint64 ByteCount, uint64 ReadCount) {};
+	void 	OnDecompressQueued(uint64 CompressedSize, uint64 UncompressedSize) {};
+	void 	OnDecompressComplete(uint64 CompressedSize, uint64 UncompressedSize) {};
+	void	OnBytesScattered(int64 BytesScattered) {};
+	void 	OnBlockCacheStore(uint64 NumBytes) {};
+	void 	OnBlockCacheHit(uint64 NumBytes) {};
+	void 	OnBlockCacheMiss(uint64 NumBytes) {};
+	void	OnSequentialRead() {};
+	void 	OnSeek(uint64 LastOffset, uint64 NewOffset) {};
+	void 	OnHandleChangeSeek() {};
+	void 	OnTocMounted(uint64 AllocatedSize) {};
+	void 	OnTocUnmounted(uint64 AllocatedSize) {};
+	void 	OnBufferReleased() {};
+	void 	OnBufferAllocated() {};
+};
+
+#endif // UE_IO_DISPATCHER_FILESYSTEM_STATS_ENABLED 
+
+} // namesapce UE
diff --git a/Engine/Source/Runtime/PakFile/Internal/IoDispatcherFileBackendTypes.h b/Engine/Source/Runtime/PakFile/Internal/IoDispatcherFileBackendTypes.h
--- a/Engine/Source/Runtime/PakFile/Internal/IoDispatcherFileBackendTypes.h
+++ b/Engine/Source/Runtime/PakFile/Internal/IoDispatcherFileBackendTypes.h
@@ -5,10 +5,8 @@
 #include "IO/IoAllocators.h"
 #include "IO/IoStore.h"
 #include "IO/IoDispatcherBackend.h"
+#include "IO/IoDispatcherFilesystemStats.h"
 #include "Async/MappedFileHandle.h"
-#include "ProfilingDebugging/CsvProfiler.h"
-#include "ProfilingDebugging/CountersTrace.h"
-#include "Containers/Ticker.h"
 #include "Modules/ModuleInterface.h"
 
 #ifndef PLATFORM_IMPLEMENTS_IO
@@ -863,53 +861,6 @@ private:
 	friend class FFileIoStoreReader;
 	friend class FFileIoStoreBufferAllocator;
 
-#if COUNTERSTRACE_ENABLED
-	FCountersTrace::FCounterInt QueuedReadRequestsSizeCounter;
-	FCountersTrace::FCounterInt CompletedReadRequestsSizeCounter;
-	FCountersTrace::FCounterInt QueuedCompressedSizeCounter;
-	FCountersTrace::FCounterInt QueuedUncompressedSizeCounter;
-	FCountersTrace::FCounterInt CompletedCompressedSizeCounter;
-	FCountersTrace::FCounterInt CompletedUncompressedSizeCounter;
-	FCountersTrace::FCounterInt FileSystemSeeksTotalDistanceCounter;
-	FCountersTrace::FCounterInt FileSystemSeeksForwardCountCounter;
-	FCountersTrace::FCounterInt FileSystemSeeksBackwardCountCounter;
-	FCountersTrace::FCounterInt FileSystemSeeksChangeHandleCountCounter;
-	FCountersTrace::FCounterInt FileSystemCompletedRequestsSizeCounter;
-	FCountersTrace::FCounterInt BlockCacheStoredSizeCounter;
-	FCountersTrace::FCounterInt BlockCacheHitSizeCounter;
-	FCountersTrace::FCounterInt BlockCacheMissedSizeCounter;
-	FCountersTrace::FCounterInt ScatteredSizeCounter;
-	FCountersTrace::FCounterInt TocMemoryCounter;
-	FCountersTrace::TCounter<std::atomic<int64>, TraceCounterType_Int> AvailableBuffersCounter;
-#endif
-
-#if CSV_PROFILER_STATS
-	uint64 QueuedFilesystemReadBytes = 0;
-	uint64 QueuedFilesystemReads = 0;
-
-	uint64 QueuedUncompressBytesIn = 0;
-	uint64 QueuedUncompressBytesOut = 0;
-	uint64 QueuedUncompressBlocks = 0;
-#endif
-
-	// Used for seek tracking 
-	const FFileIoStoreContainerFilePartition* LastContainerFilePartition = nullptr;
-	uint64 LastOffset = uint64(-1);
-
-	FTSTicker::FDelegateHandle TickerHandle;
-
-	float BytesToApproxMB(uint64 Bytes)
-	{
-		return float(double(Bytes) / 1024.0 / 1024.0);
-	}
-
-	float BytesToApproxKB(uint64 Bytes)
-	{
-		return float(double(Bytes) / 1024.0);
-	}
-
-	PAKFILE_API bool CsvTick(float DeltaTime);
-
 	PAKFILE_API void OnReadRequestsQueued(const FFileIoStoreReadRequestList& Requests);
 	PAKFILE_API void OnReadRequestsCompleted(const FFileIoStoreReadRequestList& Requests);
 	PAKFILE_API void OnDecompressQueued(const FFileIoStoreCompressedBlock* CompressedBlock);
@@ -934,6 +885,8 @@ private:
 
 	PAKFILE_API void OnBufferReleased();
 	PAKFILE_API void OnBufferAllocated();
+
+	UE::FIoDispatcherFilesystemStats Stats;
 };
 #else
 class FFileIoStoreStats
diff --git a/Engine/Source/Runtime/PakFile/Private/IoDispatcherFileBackend.cpp b/Engine/Source/Runtime/PakFile/Private/IoDispatcherFileBackend.cpp
--- a/Engine/Source/Runtime/PakFile/Private/IoDispatcherFileBackend.cpp
+++ b/Engine/Source/Runtime/PakFile/Private/IoDispatcherFileBackend.cpp
@@ -25,6 +25,7 @@
 #include "Modules/ModuleManager.h"
 #include "IO/IoContainerHeader.h"
 #include "IO/IoDispatcherConfig.h"
+#include "IO/IoDispatcherFilesystemStats.h"
 #include "Serialization/MemoryReader.h"
 #include "FileCache/FileCache.h"
 
@@ -2191,86 +2192,14 @@ uint32 FFileIoStore::GetThreadId() const
 	return Thread ? Thread->GetThreadID() : 0; 
 }
 
-CSV_DEFINE_CATEGORY(IoDispatcherFileBackend, true);
-CSV_DEFINE_CATEGORY(IoDispatcherFileBackendVerbose, false);
-
-// These stats go to both insights and csv by default
-// TODO: Ideally these should go to insights even if CSV is not capturing, but not be doubled-up if both CSV and Insights are capturing
-// TODO: It would also be nice to send these to insights as int64 without unit conversion where appropriate
-// IoDispatcher thread
-CSV_DEFINE_STAT(IoDispatcherFileBackend,			FrameBytesScatteredKB);
-CSV_DEFINE_STAT(IoDispatcherFileBackend,			QueuedFilesystemReadMB);
-CSV_DEFINE_STAT(IoDispatcherFileBackendVerbose,	QueuedFilesystemReads);
-CSV_DEFINE_STAT(IoDispatcherFileBackendVerbose,	QueuedUncompressBlocks);	
-CSV_DEFINE_STAT(IoDispatcherFileBackendVerbose,	QueuedUncompressInMB);		
-CSV_DEFINE_STAT(IoDispatcherFileBackendVerbose,	QueuedUncompressOutMB);	
-CSV_DEFINE_STAT(IoDispatcherFileBackendVerbose,	FrameBytesReadKB);				
-CSV_DEFINE_STAT(IoDispatcherFileBackendVerbose,	FrameBytesUncompressedInKB);	
-CSV_DEFINE_STAT(IoDispatcherFileBackendVerbose,	FrameBytesUncompressedOutKB);	
-
-// FileIoStore thread			
-CSV_DEFINE_STAT(IoDispatcherFileBackend,			FrameFilesystemBytesReadKB);	
-CSV_DEFINE_STAT(IoDispatcherFileBackend,			FrameSequentialReads);
-CSV_DEFINE_STAT(IoDispatcherFileBackend,			FrameSeeks);
-CSV_DEFINE_STAT(IoDispatcherFileBackendVerbose,	FrameFilesystemReads);
-CSV_DEFINE_STAT(IoDispatcherFileBackendVerbose,	FrameForwardSeeks);
-CSV_DEFINE_STAT(IoDispatcherFileBackendVerbose,	FrameBackwardSeeks);
-CSV_DEFINE_STAT(IoDispatcherFileBackendVerbose,	FrameHandleChangeSeeks);
-CSV_DEFINE_STAT(IoDispatcherFileBackendVerbose,	FrameSeekDistanceMB);	
-CSV_DEFINE_STAT(IoDispatcherFileBackendVerbose,	FrameBlockCacheStores);
-CSV_DEFINE_STAT(IoDispatcherFileBackendVerbose,	FrameBlockCacheStoresKB);
-CSV_DEFINE_STAT(IoDispatcherFileBackendVerbose,	FrameBlockCacheHits);	
-CSV_DEFINE_STAT(IoDispatcherFileBackendVerbose,	FrameBlockCacheHitKB);	
-CSV_DEFINE_STAT(IoDispatcherFileBackendVerbose,	FrameBlockCacheMisses);
-CSV_DEFINE_STAT(IoDispatcherFileBackendVerbose,	FrameBlockCacheMissKB);
-
 #if UE_FILEIOSTORE_STATS_ENABLED
 
 FFileIoStoreStats::FFileIoStoreStats()
-#if COUNTERSTRACE_ENABLED
-	: QueuedReadRequestsSizeCounter(TEXT("FileIoStore/QueuedReadRequestsSize"), TraceCounterDisplayHint_Memory)
-	, CompletedReadRequestsSizeCounter(TEXT("FileIoStore/CompletedReadRequestsSize"), TraceCounterDisplayHint_Memory)
-	, QueuedCompressedSizeCounter(TEXT("FileIoStore/QueuedCompressedSize"), TraceCounterDisplayHint_Memory)
-	, QueuedUncompressedSizeCounter(TEXT("FileIoStore/QueuedUncompressedSize"), TraceCounterDisplayHint_Memory)
-	, CompletedCompressedSizeCounter(TEXT("FileIoStore/CompletedCompressedSize"), TraceCounterDisplayHint_Memory)
-	, CompletedUncompressedSizeCounter(TEXT("FileIoStore/CompletedUncompressedSize"), TraceCounterDisplayHint_Memory)
-	, FileSystemSeeksTotalDistanceCounter(TEXT("FileIoStore/FileSystemSeeksTotalDistance"), TraceCounterDisplayHint_Memory)
-	, FileSystemSeeksForwardCountCounter(TEXT("FileIoStore/FileSystemSeeksForwardCount"), TraceCounterDisplayHint_None)
-	, FileSystemSeeksBackwardCountCounter(TEXT("FileIoStore/FileSystemSeeksBackwardCount"), TraceCounterDisplayHint_None)
-	, FileSystemSeeksChangeHandleCountCounter(TEXT("FileIoStore/FileSystemSeeksChangeHandleCount"), TraceCounterDisplayHint_None)
-	, FileSystemCompletedRequestsSizeCounter(TEXT("FileIoStore/FileSystemCompletedRequestsSize"), TraceCounterDisplayHint_Memory)
-	, BlockCacheStoredSizeCounter(TEXT("FileIoStore/BlockCacheStoredSize"), TraceCounterDisplayHint_Memory)
-	, BlockCacheHitSizeCounter(TEXT("FileIoStore/BlockCacheHitSize"), TraceCounterDisplayHint_Memory)
-	, BlockCacheMissedSizeCounter(TEXT("FileIoStore/BlockCacheMissedSize"), TraceCounterDisplayHint_Memory)
-	, ScatteredSizeCounter(TEXT("FileIoStore/ScatteredSize"), TraceCounterDisplayHint_Memory)
-	, TocMemoryCounter(TEXT("FileIoStore/TocMemory"), TraceCounterDisplayHint_Memory)
-	, AvailableBuffersCounter(TEXT("FileIoStore/AvailableBuffers"), TraceCounterDisplayHint_None)
-#endif
 {
-#if CSV_PROFILER_STATS
-	TickerHandle = FTSTicker::GetCoreTicker().AddTicker(FTickerDelegate::CreateRaw(this, &FFileIoStoreStats::CsvTick));
-#endif
 }
 
 FFileIoStoreStats::~FFileIoStoreStats()
 {
-#if CSV_PROFILER_STATS
-	FTSTicker::GetCoreTicker().RemoveTicker(TickerHandle);
-#endif
-}
-
-bool FFileIoStoreStats::CsvTick(float DeltaTime)
-{
-#if CSV_PROFILER_STATS
-	CSV_CUSTOM_STAT_DEFINED(QueuedFilesystemReadMB, BytesToApproxMB(QueuedFilesystemReadBytes), ECsvCustomStatOp::Set);
-	CSV_CUSTOM_STAT_DEFINED(QueuedFilesystemReads, (int32)QueuedFilesystemReads, ECsvCustomStatOp::Set);
-
-	CSV_CUSTOM_STAT_DEFINED(QueuedUncompressBlocks, int32(QueuedUncompressBlocks), ECsvCustomStatOp::Set);
-	CSV_CUSTOM_STAT_DEFINED(QueuedUncompressInMB, BytesToApproxMB(QueuedUncompressBytesIn), ECsvCustomStatOp::Set);
-	CSV_CUSTOM_STAT_DEFINED(QueuedUncompressOutMB, BytesToApproxMB(QueuedUncompressBytesOut), ECsvCustomStatOp::Set);
-#endif
-
-	return true;
 }
 
 void FFileIoStoreStats::OnReadRequestsQueued(const FFileIoStoreReadRequestList& Requests)
@@ -2283,89 +2212,33 @@ void FFileIoStoreStats::OnReadRequestsQueued(const FFileIoStoreReadRequestList&
 		TotalBytes += Request->Size;
 	}
 
-#if CSV_PROFILER_STATS
-	QueuedFilesystemReadBytes += TotalBytes;
-	QueuedFilesystemReads += NumReads;
-#endif
-
-#if COUNTERSTRACE_ENABLED
-	QueuedReadRequestsSizeCounter.Add(TotalBytes);
-#endif
+	Stats.OnReadRequestsQueued(TotalBytes, NumReads);
 }
 
 void FFileIoStoreStats::OnFilesystemReadStarted(const FFileIoStoreReadRequest* Request)
 {
-	CSV_CUSTOM_STAT_DEFINED(FrameFilesystemReads, 1, ECsvCustomStatOp::Accumulate);
-
-	if (LastContainerFilePartition != Request->ContainerFilePartition)
-	{
-		OnHandleChangeSeek();
-	}
-	else if (LastOffset == Request->Offset)
-	{
-		OnSequentialRead();
-	}
-	else
-	{
-		OnSeek(LastOffset, Request->Offset);
-	}
-	LastOffset = Request->Offset + Request->Size;
-	LastContainerFilePartition = Request->ContainerFilePartition;
-
-	CSV_CUSTOM_STAT_DEFINED(FrameFilesystemReads, 1, ECsvCustomStatOp::Accumulate);
+	Stats.OnFilesystemReadStarted(uint64(Request->ContainerFilePartition), Request->Offset, Request->Size);
 }
 
 void FFileIoStoreStats::OnFilesystemReadsStarted(const FFileIoStoreReadRequestList& Requests)
 {
-	int64 TotalBytes = 0;
-	int32 NumReads = 0;
 	for (const FFileIoStoreReadRequest* Request : Requests)
 	{
-		if (LastContainerFilePartition != Request->ContainerFilePartition)
-		{
-			OnHandleChangeSeek();
-		}
-		else if (LastOffset == Request->Offset)
-		{
-			OnSequentialRead();
-		}
-		else
-		{
-			OnSeek(LastOffset, Request->Offset);
-		}
-
-		LastOffset = Request->Offset + Request->Size;
-		LastContainerFilePartition = Request->ContainerFilePartition;
-
-		++NumReads;
-		TotalBytes += Request->Size;
+		Stats.OnFilesystemReadStarted(uint64(Request->ContainerFilePartition), Request->Offset, Request->Size);
 	}
-
-	CSV_CUSTOM_STAT_DEFINED(FrameFilesystemReads, NumReads, ECsvCustomStatOp::Accumulate);
 }
 
 void FFileIoStoreStats::OnFilesystemReadCompleted(const FFileIoStoreReadRequest* CompletedRequest)
 {
-	CSV_CUSTOM_STAT_DEFINED(FrameFilesystemBytesReadKB, BytesToApproxKB(CompletedRequest->Size), ECsvCustomStatOp::Accumulate);
-#if COUNTERSTRACE_ENABLED
-	FileSystemCompletedRequestsSizeCounter.Add(CompletedRequest->Size);
-#endif
+	Stats.OnFilesystemReadCompleted(uint64(CompletedRequest->ContainerFilePartition), CompletedRequest->Offset, CompletedRequest->Size);
 }
 
 void FFileIoStoreStats::OnFilesystemReadsCompleted(const FFileIoStoreReadRequestList& CompletedRequests)
 {
-	int64 TotalBytes = 0;
-	int32 NumReads = 0;
 	for (const FFileIoStoreReadRequest* Request : CompletedRequests)
 	{
-		++NumReads;
-		TotalBytes += Request->Size;
+		Stats.OnFilesystemReadCompleted(uint64(Request->ContainerFilePartition), Request->Offset, Request->Size);
 	}
-
-	CSV_CUSTOM_STAT_DEFINED(FrameFilesystemBytesReadKB, BytesToApproxKB(TotalBytes), ECsvCustomStatOp::Accumulate);
-#if COUNTERSTRACE_ENABLED
-	FileSystemCompletedRequestsSizeCounter.Add(TotalBytes);
-#endif
 }
 
 void FFileIoStoreStats::OnReadRequestsCompleted(const FFileIoStoreReadRequestList& CompletedRequests)
@@ -2377,162 +2250,57 @@ void FFileIoStoreStats::OnReadRequestsCompleted(const FFileIoStoreReadRequestLis
 		++NumReads;
 		TotalBytes += Request->Size;
 	}
-
-#if CSV_PROFILER_STATS
-	QueuedFilesystemReadBytes -= TotalBytes;
-	QueuedFilesystemReads -= NumReads;
-#endif
-
-	CSV_CUSTOM_STAT_DEFINED(FrameBytesReadKB, BytesToApproxKB(TotalBytes), ECsvCustomStatOp::Accumulate); // TODO: Send to insights if CSV is disabled
-
-#if COUNTERSTRACE_ENABLED
-	QueuedReadRequestsSizeCounter.Subtract(TotalBytes);
-	CompletedReadRequestsSizeCounter.Add(TotalBytes);
-#endif
+	Stats.OnReadRequestsCompleted(TotalBytes, NumReads);
 }
 
 void FFileIoStoreStats::OnDecompressQueued(const FFileIoStoreCompressedBlock* CompressedBlock)
 {
-#if CSV_PROFILER_STATS
-	++QueuedUncompressBlocks;
-	QueuedUncompressBytesIn += CompressedBlock->CompressedSize;
-	QueuedUncompressBytesOut += CompressedBlock->UncompressedSize;
-#endif
-
-#if COUNTERSTRACE_ENABLED
-	QueuedCompressedSizeCounter.Add(CompressedBlock->CompressedSize);
-	QueuedUncompressedSizeCounter.Add(CompressedBlock->UncompressedSize);
-#endif
+	Stats.OnDecompressQueued(CompressedBlock->CompressedSize, CompressedBlock->UncompressedSize);
 }
 
 void FFileIoStoreStats::OnDecompressComplete(const FFileIoStoreCompressedBlock* CompressedBlock)
 {
-#if CSV_PROFILER_STATS
-	--QueuedUncompressBlocks;
-	QueuedUncompressBytesIn -= CompressedBlock->CompressedSize;
-	QueuedUncompressBytesOut -= CompressedBlock->UncompressedSize;
-#endif
-
-	CSV_CUSTOM_STAT_DEFINED(FrameBytesUncompressedInKB, BytesToApproxKB(CompressedBlock->CompressedSize), ECsvCustomStatOp::Accumulate); 
-	CSV_CUSTOM_STAT_DEFINED(FrameBytesUncompressedOutKB, BytesToApproxKB(CompressedBlock->UncompressedSize), ECsvCustomStatOp::Accumulate);
-
-#if COUNTERSTRACE_ENABLED
-	QueuedCompressedSizeCounter.Subtract(CompressedBlock->CompressedSize);
-	QueuedUncompressedSizeCounter.Subtract(CompressedBlock->UncompressedSize);
-
-	CompletedCompressedSizeCounter.Add(CompressedBlock->CompressedSize);
-	CompletedUncompressedSizeCounter.Add(CompressedBlock->UncompressedSize);
-#endif
+	Stats.OnDecompressComplete(CompressedBlock->CompressedSize, CompressedBlock->UncompressedSize);
 }
 
 void FFileIoStoreStats::OnBytesScattered(int64 NumBytes)
 {
-	CSV_CUSTOM_STAT_DEFINED(FrameBytesScatteredKB, BytesToApproxKB(NumBytes), ECsvCustomStatOp::Accumulate);
-#if COUNTERSTRACE_ENABLED
-	ScatteredSizeCounter.Add(NumBytes);
-#endif
-}
-
-void FFileIoStoreStats::OnSequentialRead()
-{
-	CSV_CUSTOM_STAT_DEFINED(FrameSequentialReads, 1, ECsvCustomStatOp::Accumulate);
-}
-
-void FFileIoStoreStats::OnSeek(uint64 PrevOffset, uint64 NewOffset)
-{
-	if (NewOffset > PrevOffset)
-	{
-		int64 Delta = NewOffset - PrevOffset;
-
-		CSV_CUSTOM_STAT_DEFINED(FrameForwardSeeks, 1, ECsvCustomStatOp::Accumulate);
-		CSV_CUSTOM_STAT_DEFINED(FrameSeekDistanceMB, BytesToApproxMB(Delta), ECsvCustomStatOp::Accumulate);
-
-#if COUNTERSTRACE_ENABLED
-		FileSystemSeeksTotalDistanceCounter.Add(Delta);
-		FileSystemSeeksForwardCountCounter.Increment();
-#endif
-	}
-	else
-	{
-		int64 Delta = PrevOffset - NewOffset;
-		CSV_CUSTOM_STAT_DEFINED(FrameBackwardSeeks, 1, ECsvCustomStatOp::Accumulate);
-		CSV_CUSTOM_STAT_DEFINED(FrameSeekDistanceMB, BytesToApproxMB(Delta), ECsvCustomStatOp::Accumulate);
-
-#if COUNTERSTRACE_ENABLED
-		FileSystemSeeksTotalDistanceCounter.Add(Delta);
-		FileSystemSeeksBackwardCountCounter.Increment();
-#endif
-	}
-
-	CSV_CUSTOM_STAT_DEFINED(FrameSeeks, 1, ECsvCustomStatOp::Accumulate);
-}
-
-void FFileIoStoreStats::OnHandleChangeSeek()
-{
-	CSV_CUSTOM_STAT_DEFINED(FrameHandleChangeSeeks, 1, ECsvCustomStatOp::Accumulate);
-	CSV_CUSTOM_STAT_DEFINED(FrameSeeks, 1, ECsvCustomStatOp::Accumulate);
-
-#if COUNTERSTRACE_ENABLED
-	FileSystemSeeksChangeHandleCountCounter.Increment();
-#endif
+	Stats.OnBytesScattered(NumBytes);
 }
 
 void FFileIoStoreStats::OnBlockCacheStore(uint64 NumBytes)
 {
-	CSV_CUSTOM_STAT_DEFINED(FrameBlockCacheStores, 1, ECsvCustomStatOp::Accumulate);
-	CSV_CUSTOM_STAT_DEFINED(FrameBlockCacheStoresKB, BytesToApproxKB(NumBytes), ECsvCustomStatOp::Accumulate);
-
-#if COUNTERSTRACE_ENABLED
-	BlockCacheStoredSizeCounter.Add(NumBytes);
-#endif
+	Stats.OnBlockCacheStore(NumBytes);
 }
 
 void FFileIoStoreStats::OnBlockCacheHit(uint64 NumBytes)
 {
-	CSV_CUSTOM_STAT_DEFINED(FrameBlockCacheHits, 1, ECsvCustomStatOp::Accumulate);
-	CSV_CUSTOM_STAT_DEFINED(FrameBlockCacheHitKB, BytesToApproxKB(NumBytes), ECsvCustomStatOp::Accumulate);
-
-#if COUNTERSTRACE_ENABLED
-	BlockCacheHitSizeCounter.Add(NumBytes);
-#endif
+	Stats.OnBlockCacheHit(NumBytes);
 }
 
 void FFileIoStoreStats::OnBlockCacheMiss(uint64 NumBytes)
 {
-	CSV_CUSTOM_STAT_DEFINED(FrameBlockCacheMisses, 1, ECsvCustomStatOp::Accumulate);
-	CSV_CUSTOM_STAT_DEFINED(FrameBlockCacheMissKB, BytesToApproxKB(NumBytes), ECsvCustomStatOp::Accumulate);
-
-#if COUNTERSTRACE_ENABLED
-	BlockCacheMissedSizeCounter.Add(NumBytes);
-#endif
+	Stats.OnBlockCacheMiss(NumBytes);
 }
 
 void FFileIoStoreStats::OnTocMounted(uint64 AllocatedSize)
 {
-#if COUNTERSTRACE_ENABLED
-	TocMemoryCounter.Add(AllocatedSize);
-#endif
+	Stats.OnTocMounted(AllocatedSize);
 }
 
 void FFileIoStoreStats::OnTocUnmounted(uint64 AllocatedSize)
 {
-#if COUNTERSTRACE_ENABLED
-	TocMemoryCounter.Subtract(AllocatedSize);
-#endif
+	Stats.OnTocUnmounted(AllocatedSize);
 }
 
 void FFileIoStoreStats::OnBufferReleased()
 {
-#if COUNTERSTRACE_ENABLED
-	AvailableBuffersCounter.Increment();
-#endif
+	Stats.OnBufferReleased();
 }
 
 void FFileIoStoreStats::OnBufferAllocated()
 {
-#if COUNTERSTRACE_ENABLED
-	AvailableBuffersCounter.Decrement();
-#endif
+	Stats.OnBufferAllocated();
 }
 
 #endif
