FPakFile: Use FStringView based SplitPathInline to avoid FString allocations. Eliminates 3M+ short living allocs during a client boot in an internal project.

diff --git a/Engine/Source/Runtime/PakFile/Private/PakFile.cpp b/Engine/Source/Runtime/PakFile/Private/PakFile.cpp
--- a/Engine/Source/Runtime/PakFile/Private/PakFile.cpp
+++ b/Engine/Source/Runtime/PakFile/Private/PakFile.cpp
@@ -1068,20 +1068,25 @@ void FPakFile::PruneDirectoryIndexInternal(FDirectoryIndex* InOutDirectoryIndex,
 			{
 				if (Pair.Value)
 				{
-					FString CurrentDirectory = Pair.Key;
-					FString UnusedCleanFileName;
+					FStringView CurrentDirectory = Pair.Key;
+					FStringView UnusedCleanFileName;
 					while (CurrentDirectory != MountPoint)
 					{
 						if (!SplitPathInline(CurrentDirectory, UnusedCleanFileName))
 						{
 							break;
 						}
-						bool& bOldValue = KeepDirectory.FindOrAdd(CurrentDirectory);
-						if (bOldValue)
+						const uint32 CurrentDirectoryHash = GetTypeHash(CurrentDirectory);
+						bool* bOldValue = KeepDirectory.FindByHash(CurrentDirectoryHash, CurrentDirectory);
+						if (!bOldValue)
+						{
+							bOldValue = &KeepDirectory.AddByHash(CurrentDirectoryHash, FString(CurrentDirectory));
+						}
+						if (*bOldValue)
 						{
 							break;
 						}
-						bOldValue = true;
+						*bOldValue = true;
 					}
 				}
 			}
@@ -1186,20 +1191,25 @@ void FPakFile::PruneDirectoryIndexInternal(FDirectoryIndex* InOutDirectoryIndex,
 			{
 				if (Pair.Value)
 				{
-					FString CurrentDirectory = Pair.Key;
-					FString UnusedCleanFileName;
+					FStringView CurrentDirectory = Pair.Key;
+					FStringView UnusedCleanFileName;
 					while (CurrentDirectory != MountPoint)
 					{
 						if (!SplitPathInline(CurrentDirectory, UnusedCleanFileName))
 						{
 							break;
 						}
-						bool& bOldValue = KeepDirectory.FindOrAdd(CurrentDirectory);
-						if (bOldValue)
+						const uint32 CurrentDirectoryHash = GetTypeHash(CurrentDirectory);
+						bool* bOldValue = KeepDirectory.FindByHash(CurrentDirectoryHash, CurrentDirectory);
+						if (!bOldValue)
+						{
+							bOldValue = &KeepDirectory.AddByHash(CurrentDirectoryHash, FString(CurrentDirectory));
+						}
+						if (*bOldValue)
 						{
 							break;
 						}
-						bOldValue = true;
+						*bOldValue = true;
 					}
 				}
 			}
@@ -1388,45 +1398,48 @@ void FPakFile::AddEntryToIndex(const FString& Filename, const FPakEntryLocation&
 
 	if (DirectoryIndex)
 	{
-		FString RelativeDirectoryFromMount = RelativePathFromMount;
-		FString CleanFileName;
+		FStringView RelativeDirectoryFromMount = RelativePathFromMount;
+		FStringView CleanFileName;
 		SplitPathInline(RelativeDirectoryFromMount, CleanFileName);
-		FPakDirectory* Directory = DirectoryIndex->Find(RelativeDirectoryFromMount);
+		const uint32 RelativeDirectoryFromMountHash = GetTypeHash(RelativeDirectoryFromMount);
+		FPakDirectory* Directory = DirectoryIndex->FindByHash(RelativeDirectoryFromMountHash, RelativeDirectoryFromMount);
 		if (Directory == nullptr)
 		{
 			// add the parent directories up to the mount point (mount point relative path is "/")
-			FString CurrentDirectory(RelativeDirectoryFromMount);
-			FString UnusedCleanFileName;
+			FStringView CurrentDirectory(RelativeDirectoryFromMount);
+			FStringView UnusedCleanFileName;
 			while (!CurrentDirectory.IsEmpty())
 			{
-				// FPaths::GetPath doesn't handle our / at the end of directories, so call FindLastChar(/)
 				if (!SplitPathInline(CurrentDirectory, UnusedCleanFileName))
 				{
 					break;
 				}
-				DirectoryIndex->FindOrAdd(CurrentDirectory);
+				const uint32 CurrentDirectoryHash = GetTypeHash(CurrentDirectory);
+				if (!DirectoryIndex->FindByHash(CurrentDirectoryHash, CurrentDirectory))
+				{
+					DirectoryIndex->AddByHash(CurrentDirectoryHash, FString(CurrentDirectory));
+				}
 			}
 
 			// Add the new directory; this has to come after the addition of the parent directories because we want to use the pointer afterwards and adding other directories would invalidate it
-			Directory = &DirectoryIndex->Add(RelativeDirectoryFromMount);
+			Directory = &DirectoryIndex->AddByHash(RelativeDirectoryFromMountHash, FString(RelativeDirectoryFromMount));
 		}
-		Directory->Add(CleanFileName, EntryLocation);
+		Directory->Add(FString(CleanFileName), EntryLocation);
 	}
 #if ENABLE_PAKFILE_USE_DIRECTORY_TREE
 	if (DirectoryTreeIndex)
 	{
-		FString RelativeDirectoryFromMount = RelativePathFromMount;
-		FString CleanFileName;
+		FStringView RelativeDirectoryFromMount = RelativePathFromMount;
+		FStringView CleanFileName;
 		SplitPathInline(RelativeDirectoryFromMount, CleanFileName);
 		FPakDirectory* Directory = DirectoryTreeIndex->Find(RelativeDirectoryFromMount);
 		if (Directory == nullptr)
 		{
 			// add the parent directories up to the mount point (mount point relative path is "/")
-			FString CurrentDirectory(RelativeDirectoryFromMount);
-			FString UnusedCleanFileName;
+			FStringView CurrentDirectory(RelativeDirectoryFromMount);
+			FStringView UnusedCleanFileName;
 			while (!CurrentDirectory.IsEmpty())
 			{
-				// FPaths::GetPath doesn't handle our / at the end of directories, so call FindLastChar(/)
 				if (!SplitPathInline(CurrentDirectory, UnusedCleanFileName))
 				{
 					break;
@@ -1437,7 +1450,7 @@ void FPakFile::AddEntryToIndex(const FString& Filename, const FPakEntryLocation&
 			// Add the new directory; this has to come after the addition of the parent directories because we want to use the pointer afterwards and adding other directories would invalidate it
 			Directory = &DirectoryTreeIndex->FindOrAdd(RelativeDirectoryFromMount);
 		}
-		Directory->Add(CleanFileName, EntryLocation);
+		Directory->Add(FString(CleanFileName), EntryLocation);
 	}
 #endif
 
@@ -2171,18 +2184,18 @@ const FPakEntryLocation* FPakFile::FindLocationFromIndex(const FString& FullPath
 	{
 		return nullptr;
 	}
-	const TCHAR* RelativePathFromMount = (*FullPath) + MountPoint.Len();
-	FString RelativeDirName(RelativePathFromMount);
-	FString CleanFileName;
+	FStringView RelativePathFromMount = FStringView(FullPath).Mid(MountPoint.Len());
+	FStringView RelativeDirName(RelativePathFromMount);
+	FStringView CleanFileName;
 	if (RelativeDirName.IsEmpty())
 	{
 		return nullptr;
 	}
 	SplitPathInline(RelativeDirName, CleanFileName);
-	const FPakDirectory* PakDirectory = DirectoryIndex.Find(RelativeDirName);
+	const FPakDirectory* PakDirectory = DirectoryIndex.FindByHash(GetTypeHash(RelativeDirName), RelativeDirName);
 	if (PakDirectory)
 	{
-		return PakDirectory->Find(CleanFileName);
+		return PakDirectory->FindByHash(GetTypeHash(CleanFileName), CleanFileName);
 	}
 	return nullptr;
 }
@@ -2196,16 +2209,16 @@ const FPakEntryLocation* FPakFile::FindLocationFromIndex(const FString& FullPath
 		const FPakEntryLocation* Result = nullptr;
 		if (FullPath.StartsWith(MountPoint))
 		{
-			const TCHAR* RelativePathFromMount = (*FullPath) + MountPoint.Len();
-			FString RelativeDirName(RelativePathFromMount);
-			FString CleanFileName;
+			FStringView RelativePathFromMount = FStringView(FullPath).Mid(MountPoint.Len());
+			FStringView RelativeDirName(RelativePathFromMount);
+			FStringView CleanFileName;
 			if (!RelativeDirName.IsEmpty())
 			{
 				SplitPathInline(RelativeDirName, CleanFileName);
 				const FPakDirectory* PakDirectory = InDirectoryTreeIndex.Find(RelativeDirName);
 				if (PakDirectory)
 				{
-					Result = PakDirectory->Find(CleanFileName);
+					Result = PakDirectory->FindByHash(GetTypeHash(CleanFileName), CleanFileName);
 				}
 			}
 		}
diff --git a/Engine/Source/Runtime/PakFile/Public/IPlatformFilePak.h b/Engine/Source/Runtime/PakFile/Public/IPlatformFilePak.h
--- a/Engine/Source/Runtime/PakFile/Public/IPlatformFilePak.h
+++ b/Engine/Source/Runtime/PakFile/Public/IPlatformFilePak.h
@@ -1428,7 +1428,7 @@ public:
 	}
 
 	/** Helper function to split a PakDirectoryIndex-Formatted PathName into its PakDirectoryIndex-Formatted parent directory and the CleanFileName */
-	static bool SplitPathInline(FString& InOutPath, FString& OutFilename)
+	static bool SplitPathInline(FStringView& InOutPath, FStringView& OutFilename)
 	{
 		// FPaths::GetPath doesn't handle our / at the end of directories, so we have to do string manipulation ourselves
 		// The manipulation is less complicated than GetPath deals with, since we have normalized/path/strings, we have relative paths only, and we don't care about extensions
@@ -1442,7 +1442,7 @@ public:
 			if (InOutPath[0] == TEXT('/'))
 			{
 				// The root directory; it has no parent.
-				OutFilename.Empty();
+				OutFilename = FStringView();
 				return false;
 			}
 			else
@@ -1458,7 +1458,7 @@ public:
 			if (InOutPath[InOutPath.Len() - 1] == TEXT('/'))
 			{
 				// The input was a Directory; remove the trailing / since we don't keep those on the CleanFilename
-				InOutPath.LeftChopInline(1, EAllowShrinking::No);
+				InOutPath.LeftChopInline(1);
 			}
 
 			int32 Offset = 0;
@@ -1466,7 +1466,7 @@ public:
 			{
 				int32 FilenameStart = Offset + 1;
 				OutFilename = InOutPath.Mid(FilenameStart);
-				InOutPath.LeftInline(FilenameStart, EAllowShrinking::No); // The Parent Directory keeps the / at the end
+				InOutPath.LeftInline(FilenameStart); // The Parent Directory keeps the / at the end
 			}
 			else
 			{
@@ -1478,6 +1478,15 @@ public:
 		}
 	}
 
+	static bool SplitPathInline(FString& InOutPath, FString& OutFilename)
+	{
+		FStringView Path(InOutPath), Filename;
+		const bool bResult = SplitPathInline(Path, Filename);
+		OutFilename = Filename;
+		InOutPath = Path;
+		return bResult;
+	}
+
 	/**
 	 * Helper function to return Child's relative path from the mount point.  Returns false if Child is not equal to MountPoint and is not a child path under MountPoint, else returns true.
 	 * Edits Child only if returning true, setting it to the relative path.
