PakFile: Address static analysis warnings: `V789: Iterators for the 'KeepDirectory' container, used in the range-based for loop, become invalid upon the call of the 'AddByHash' function.`

diff --git a/Engine/Source/Runtime/PakFile/Private/PakFile.cpp b/Engine/Source/Runtime/PakFile/Private/PakFile.cpp
--- a/Engine/Source/Runtime/PakFile/Private/PakFile.cpp
+++ b/Engine/Source/Runtime/PakFile/Private/PakFile.cpp
@@ -1063,12 +1063,22 @@ void FPakFile::PruneDirectoryIndexInternal(FDirectoryIndex* InOutDirectoryIndex,
 				KeepDirectory.FindOrAdd(DirectoryName) = bKeepDirectory;
 			}
 
-			// For each kept directory, mark that we need to keep all of its parents up to the mount point
-			for (const TPair<FString, bool>& Pair : KeepDirectory)
 			{
-				if (Pair.Value)
+				// For each kept directory, mark that we need to keep all of its parents up to the mount point.
+				// Note: We rely on TMap reallocations for KeepDirectory not modifying the underlying FString data,
+				// so the FStringViews in KeptDirectories remain valid.
+				TArray<FStringView> KeptDirectories;
+				for (const TPair<FString, bool>& Pair : KeepDirectory)
+				{
+					if (Pair.Value)
+					{
+						KeptDirectories.Add(Pair.Key);
+					}
+				}
+
+				for (const FStringView& KeptDirectory : KeptDirectories)
 				{
-					FStringView CurrentDirectory = Pair.Key;
+					FStringView CurrentDirectory = KeptDirectory;
 					FStringView UnusedCleanFileName;
 					while (CurrentDirectory != MountPoint)
 					{
@@ -1186,12 +1196,22 @@ void FPakFile::PruneDirectoryIndexInternal(FDirectoryIndex* InOutDirectoryIndex,
 				KeepDirectory.FindOrAdd(DirectoryName) = bKeepDirectory;
 			}
 
-			// For each kept directory, mark that we need to keep all of its parents up to the mount point
-			for (const TPair<FString, bool>& Pair : KeepDirectory)
 			{
-				if (Pair.Value)
+				// For each kept directory, mark that we need to keep all of its parents up to the mount point.
+				// Note: We rely on TMap reallocations for KeepDirectory not modifying the underlying FString data,
+				// so the FStringViews in KeptDirectories remain valid.
+				TArray<FStringView> KeptDirectories;
+				for (const TPair<FString, bool>& Pair : KeepDirectory)
+				{
+					if (Pair.Value)
+					{
+						KeptDirectories.Add(Pair.Key);
+					}
+				}
+
+				for (const FStringView& KeptDirectory : KeptDirectories)
 				{
-					FStringView CurrentDirectory = Pair.Key;
+					FStringView CurrentDirectory = KeptDirectory;
 					FStringView UnusedCleanFileName;
 					while (CurrentDirectory != MountPoint)
 					{
