[Core] Shrunk some containers in AsyncLoading code that had a wasted slack.

#rnx

diff --git a/Engine/Source/Runtime/PakFile/Private/FilePackageStore.cpp b/Engine/Source/Runtime/PakFile/Private/FilePackageStore.cpp
--- a/Engine/Source/Runtime/PakFile/Private/FilePackageStore.cpp
+++ b/Engine/Source/Runtime/PakFile/Private/FilePackageStore.cpp
@@ -705,6 +705,11 @@ void FFilePackageStoreBackend::Update()
 
 		// Finally reallocate entry data and build new map
 		EntryData.Shrink();
+		LocalizedPackages.Shrink();
+		RedirectsPackageMap.Shrink();
+#if WITH_EDITOR
+		OptionalSegmentStoreEntriesMap.Shrink();
+#endif
 		PackageEntries = FPackageIdMap(MoveTemp(Pairs));
 
 		bNeedsContainerUpdate = false;	
