Remove unused header in file I/O dispatcher backend

#rnx

diff --git a/Engine/Source/Runtime/PakFile/Private/FileIoDispatcherBackend.cpp b/Engine/Source/Runtime/PakFile/Private/FileIoDispatcherBackend.cpp
--- a/Engine/Source/Runtime/PakFile/Private/FileIoDispatcherBackend.cpp
+++ b/Engine/Source/Runtime/PakFile/Private/FileIoDispatcherBackend.cpp
@@ -8,7 +8,6 @@
 #include "HAL/CriticalSection.h"
 #include "HAL/Platform.h"
 #include "HAL/PlatformFileManager.h"
-#include "IO/GenericPlatformIoDispatcher.h"
 #include "IO/IoAllocators.h"
 #include "IO/IoChunkId.h"
 #include "IO/IoChunkEncoding.h"
