Add support for retrieving soft referenced package ID's for containers saved with soft references.

#rnx

diff --git a/Engine/Source/Runtime/Core/Private/IO/PackageStore.cpp b/Engine/Source/Runtime/Core/Private/IO/PackageStore.cpp
--- a/Engine/Source/Runtime/Core/Private/IO/PackageStore.cpp
+++ b/Engine/Source/Runtime/Core/Private/IO/PackageStore.cpp
@@ -214,6 +214,19 @@ bool FPackageStore::GetPackageRedirectInfo(FPackageId PackageId, FName& OutSourc
 	return false;
 }
 
+TConstArrayView<uint32> FPackageStore::GetSoftReferences(FPackageId PackageId, TConstArrayView<FPackageId>& OutPackageIds)
+{
+	check(ThreadReadCount);
+	for (const FBackendAndPriority& Backend : Backends)
+	{
+		if (TConstArrayView<uint32> SoftRefs = Backend.Value->GetSoftReferences(PackageId, OutPackageIds); !SoftRefs.IsEmpty())
+		{
+			return SoftRefs;
+		}
+	}
+	return TConstArrayView<uint32>();
+}
+
 FPackageStoreBackendContext::FPendingEntriesAddedEvent& FPackageStore::OnPendingEntriesAdded()
 {
 	return BackendContext->PendingEntriesAdded;
diff --git a/Engine/Source/Runtime/Core/Public/IO/PackageStore.h b/Engine/Source/Runtime/Core/Public/IO/PackageStore.h
--- a/Engine/Source/Runtime/Core/Public/IO/PackageStore.h
+++ b/Engine/Source/Runtime/Core/Public/IO/PackageStore.h
@@ -153,6 +153,12 @@ public:
 
 	/* Returns the redirected package ID and source package name for the specified package ID if it's being redirected. */
 	virtual bool GetPackageRedirectInfo(FPackageId PackageId, FName& OutSourcePackageName, FPackageId& OutRedirectedToPackageId) = 0;
+
+	/* Returns all soft referenced package IDs for the specified package ID. */
+	virtual TConstArrayView<uint32> GetSoftReferences(FPackageId PackageId, TConstArrayView<FPackageId>& OutPackageIds)
+	{
+		return TConstArrayView<uint32>();
+	}
 };
 
 /**
@@ -173,6 +179,9 @@ public:
 	/* Returns the redirected package ID and source package name for the specified package ID if it's being redirected. */
 	CORE_API bool GetPackageRedirectInfo(FPackageId PackageId, FName& OutSourcePackageName, FPackageId& OutRedirectedToPackageId);
 
+	/* Returns all soft referenced package IDs for the specified package ID. */
+	CORE_API TConstArrayView<uint32> GetSoftReferences(FPackageId PackageId, TConstArrayView<FPackageId>& OutPackageIds);
+
 	CORE_API FPackageStoreBackendContext::FPendingEntriesAddedEvent& OnPendingEntriesAdded();
 
 	CORE_API bool HasAnyBackendsMounted() const;
diff --git a/Engine/Source/Runtime/PakFile/Private/FilePackageStore.cpp b/Engine/Source/Runtime/PakFile/Private/FilePackageStore.cpp
--- a/Engine/Source/Runtime/PakFile/Private/FilePackageStore.cpp
+++ b/Engine/Source/Runtime/PakFile/Private/FilePackageStore.cpp
@@ -1,5 +1,6 @@
 // Copyright Epic Games, Inc. All Rights Reserved.
 #include "FilePackageStore.h"
+#include "Algo/Find.h"
 #include "IO/IoContainerId.h"
 #include "IO/IoContainerHeader.h"
 #include "Internationalization/PackageLocalizationManager.h"
@@ -384,6 +385,28 @@ bool FFilePackageStoreBackend::GetPackageRedirectInfo(FPackageId PackageId, FNam
 	return false;
 }
 
+TConstArrayView<uint32> FFilePackageStoreBackend::GetSoftReferences(FPackageId PackageId, TConstArrayView<FPackageId>& OutPackageIds)
+{
+	if (const FEntryHandle* Entry = PackageEntries.Find(PackageId); Entry != nullptr && Entry->HasSoftRefs)
+	{
+		FMountedContainer* Container = Algo::FindByPredicate(
+			MountedContainers,
+			[Entry](const FMountedContainer& C) { return C.EntryDataRange.Contains(Entry->Offset); });
+
+		if (Container != nullptr)
+		{
+			OutPackageIds = Container->ContainerHeader->SoftPackageReferences.PackageIds;
+			return GetSoftReferenceIndices(*Entry);
+		}
+		else
+		{
+			UE_LOG(LogFilePackageStore, Warning, TEXT("Failed to find container with soft references for package ID '%s'"), *LexToString(PackageId));
+		}
+	}
+
+	return TConstArrayView<uint32>();
+}
+
 void FFilePackageStoreBackend::Mount(FIoContainerHeader* ContainerHeader, uint32 Order)
 {
 	LLM_SCOPE(ELLMTag::AsyncLoading);
@@ -425,7 +448,7 @@ constexpr uint32 GetNumWords(uint32 Num)
 	return (sizeof(T) / sizeof(uint32)) * Num;
 }
 
-FFilePackageStoreBackend::FEntryHandle FFilePackageStoreBackend::AddNewEntryData(const FFilePackageStoreEntry& Entry)
+FFilePackageStoreBackend::FEntryHandle FFilePackageStoreBackend::AddNewEntryData(const FFilePackageStoreEntry& Entry, const FFilePackageStoreEntrySoftReferences* SoftReferences)
 {
 	check(EntryData.Num() < (1u << FEntryHandle::OffsetBits));
 	FEntryHandle Out;
@@ -454,6 +477,12 @@ FFilePackageStoreBackend::FEntryHandle FFilePackageStoreBackend::AddNewEntryData
 			Out.HasShaderMaps = 1;
 			AppendEntryData(Entry.ShaderMapHashes.Data(), Num);
 		}
+
+		if (uint32 Num = SoftReferences != nullptr ? SoftReferences->Indices.Num() : 0)
+		{
+			Out.HasSoftRefs = 1;
+			AppendEntryData(SoftReferences->Indices.Data(), Num);
+		}
 	}
 
 	check((uint32)EntryData.Num() > Out.Offset);
@@ -482,6 +511,12 @@ FFilePackageStoreBackend::FEntryHandle FFilePackageStoreBackend::AddOldEntryData
 		uint32 NumShaderWords = OldHandle.HasShaderMaps ? 1 + GetNumWords<FSHAHash>(*It)  : 0;
 		check(It + NumShaderWords <= OldEntryData.end())
 		EntryData.Append(It, NumShaderWords);
+
+		It += NumShaderWords;
+
+		uint32 NumSoftRefWords = OldHandle.HasSoftRefs ? 1 + GetNumWords<uint32>(*It)  : 0;
+		check(It + NumSoftRefWords <= OldEntryData.end())
+		EntryData.Append(It, NumSoftRefWords);
 	}
 	
 	check((uint32)EntryData.Num() > Out.Offset);
@@ -520,10 +555,27 @@ TConstArrayView<FSHAHash> FFilePackageStoreBackend::GetShaderHashes(FEntryHandle
 	return MakeEntryDataSlice<FSHAHash>(Data + NumPackageIdWords, EntryData);
 }
 
+TConstArrayView<uint32> FFilePackageStoreBackend::GetSoftReferenceIndices(FEntryHandle Handle)
+{
+	if (Handle.HasSoftRefs == 0)
+	{
+		return {};
+	}
+
+	const uint32* Data = &EntryData[Handle.Offset];
+	uint32 NumPackageIdAndHashWords = Handle.HasPackageIds ? 1 + GetNumWords<FPackageId>(*Data) : 0;
+	if (Handle.HasShaderMaps)
+	{
+		NumPackageIdAndHashWords += (1 + GetNumWords<FSHAHash>(*(Data + NumPackageIdAndHashWords)));
+	}
+	return MakeEntryDataSlice<uint32>(Data + NumPackageIdAndHashWords, EntryData);
+}
+
 // Update() entry deduplication helper
 struct FFilePackageStoreEntryRef
 {
 	const FFilePackageStoreEntry* Entry;
+	const FFilePackageStoreEntrySoftReferences* SoftReferences = nullptr;
 
 	FMemoryView GetImportedPackages() const
 	{
@@ -535,9 +587,14 @@ struct FFilePackageStoreEntryRef
 		return MakeMemoryView(Entry->ShaderMapHashes.begin(), Entry->ShaderMapHashes.end());
 	}
 
+	FMemoryView GetSoftReferences() const
+	{
+		return SoftReferences != nullptr ? FMemoryView(SoftReferences->Indices.begin(), SoftReferences->Indices.end()) : FMemoryView();
+	}
+
 	bool operator==(FFilePackageStoreEntryRef O) const
 	{
-		return GetImportedPackages().EqualBytes(O.GetImportedPackages()) && GetShaderHashes().EqualBytes(O.GetShaderHashes());
+		return GetImportedPackages().EqualBytes(O.GetImportedPackages()) && GetShaderHashes().EqualBytes(O.GetShaderHashes()) && GetSoftReferences().EqualBytes(O.GetSoftReferences());
 	}
 
 	friend uint32 GetTypeHash(FFilePackageStoreEntryRef Ref)
@@ -554,6 +611,14 @@ struct FFilePackageStoreEntryRef
 			Out = HashCombineFast(Out, *reinterpret_cast<const uint32*>(SHA1.Hash));
 		}
 
+		if (Ref.SoftReferences != nullptr)
+		{
+			for (uint32 Idx : Ref.SoftReferences->Indices)
+			{
+				Out = HashCombineFast(Out, Idx); 
+			}
+		}
+
 		return Out;
 	}
 };
@@ -635,23 +700,33 @@ void FFilePackageStoreBackend::Update()
 				// Add an empty entry per mounted container to allow inverse mapping handle offset -> container,
 				// which in turn allows releasing FIoContainerHeader::PackageIds after first Update()
 				UniqueEntriesPerNewContainer.Reset();
-				UniqueEntriesPerNewContainer.Add({&EmptyEntry}, AddNewEntryData(EmptyEntry));
+				UniqueEntriesPerNewContainer.Add({&EmptyEntry}, AddNewEntryData(EmptyEntry, nullptr));
+
+				TConstArrayView<FFilePackageStoreEntrySoftReferences> AllSoftReferences;
+				if (ContainerHeader->SoftPackageReferences.bContainsSoftPackageReferences)
+				{
+					AllSoftReferences = MakeArrayView<const FFilePackageStoreEntrySoftReferences>(
+						reinterpret_cast<const FFilePackageStoreEntrySoftReferences*>(ContainerHeader->SoftPackageReferences.PackageIndices.GetData()),
+						ContainerHeader->PackageIds.Num());
+				}
 
-				const FPackageId* PIDIt = ContainerHeader->PackageIds.GetData();
-				for (const FFilePackageStoreEntry& StoreEntry : MakeArrayView(reinterpret_cast<FFilePackageStoreEntry*>(ContainerHeader->StoreEntries.GetData()), NumNewPackages))
+				TConstArrayView<FFilePackageStoreEntry> PackageStoreEntries = MakeArrayView(reinterpret_cast<const FFilePackageStoreEntry*>(ContainerHeader->StoreEntries.GetData()), NumNewPackages);
+				for (int32 Idx = 0; const FFilePackageStoreEntry& StoreEntry : PackageStoreEntries)
 				{
-					FEntryHandle& Handle = UniqueEntriesPerNewContainer.FindOrAdd({&StoreEntry}, IllegalHandle);
+					const FFilePackageStoreEntrySoftReferences* SoftReferences = AllSoftReferences.IsEmpty() ? nullptr : &AllSoftReferences[Idx];
+					FEntryHandle& Handle = UniqueEntriesPerNewContainer.FindOrAdd({&StoreEntry, SoftReferences}, IllegalHandle);
 					if (Handle == IllegalHandle)
 					{
-						Handle = AddNewEntryData(StoreEntry);
+						Handle = AddNewEntryData(StoreEntry, SoftReferences);
 					}
 
-					Pairs.Add({*PIDIt++, Handle});
+					const FPackageId& PackageId = ContainerHeader->PackageIds[Idx++];
+					Pairs.Add({PackageId, Handle});
 				}
 
 				MountedContainer.NumMountedPackages = NumNewPackages;
 				ContainerHeader->PackageIds.Empty();
-				ContainerHeader->SoftPackageReferences.Empty();
+				ContainerHeader->SoftPackageReferences.PackageIndices.Empty();
 			}
 
 			
@@ -766,7 +841,7 @@ TEST_CASE_NAMED(FPackageIdMapTest, "System::Pak::IoStore::PackageIdMap", "[Pak][
 			{FPackageId::FromValue(0x00000000'00000002), MakeTestHandle(2)},
 			{FPackageId::FromValue(0x00000000'00000001), MakeTestHandle(1)},
 			{FPackageId::FromValue(0xfedcba09'87654321), MakeTestHandle(0x123456)},
-			{FPackageId::FromValue(0x00000000'00000003), MakeTestHandle(1u << 29)},
+			{FPackageId::FromValue(0x00000000'00000003), MakeTestHandle(1u << 28)},
 			{FPackageId::FromValue(0x00000000'10000000), MakeTestHandle(1)},
 			{FPackageId::FromValue(0x00000001'00000000), MakeTestHandle(2)},
 			{FPackageId::FromValue(0x00000010'00000000), MakeTestHandle(3)},
diff --git a/Engine/Source/Runtime/PakFile/Private/FilePackageStore.h b/Engine/Source/Runtime/PakFile/Private/FilePackageStore.h
--- a/Engine/Source/Runtime/PakFile/Private/FilePackageStore.h
+++ b/Engine/Source/Runtime/PakFile/Private/FilePackageStore.h
@@ -9,12 +9,15 @@
 
 struct FIoContainerHeader;
 struct FFilePackageStoreEntry;
+struct FFilePackageStoreEntrySoftReferences;
 
 namespace UE::FilePackageStorePrivate
 {
 	
 struct FMountedDataRange
 {
+	bool Contains(uint32 Offset) const { return Offset >= Begin && Offset <= End; }
+
 	uint32 Begin = 0;
 	uint32 End = 0;
 };
@@ -30,11 +33,12 @@ struct FMountedContainer
 
 struct FEntryHandle
 {
-	static constexpr uint32 OffsetBits = 30;
+	static constexpr uint32 OffsetBits = 29;
 
 	uint32 Offset : OffsetBits;
 	uint32 HasPackageIds : 1;
 	uint32 HasShaderMaps : 1;
+	uint32 HasSoftRefs : 1;
 
 	FEntryHandle()
 	{
@@ -81,7 +85,7 @@ struct FUncookedPackage
 /*
  * File/container based package store.
  */
-class FFilePackageStoreBackend : public IPackageStoreBackend
+class FFilePackageStoreBackend final : public IPackageStoreBackend
 {
 public:
 	virtual void OnMounted(TSharedRef<const FPackageStoreBackendContext>) override {}
@@ -89,6 +93,7 @@ public:
 	virtual void EndRead() override;
 	virtual EPackageStoreEntryStatus GetPackageStoreEntry(FPackageId PackageId, FName PackageName, FPackageStoreEntry& OutPackageStoreEntry) override;
 	virtual bool GetPackageRedirectInfo(FPackageId PackageId, FName& OutSourcePackageName, FPackageId& OutRedirectedToPackageId) override;
+	virtual TConstArrayView<uint32> GetSoftReferences(FPackageId PackageId, TConstArrayView<FPackageId>& OutPackageIds) override;
 
 	void Mount(FIoContainerHeader* ContainerHeader, uint32 Order);
 	void Unmount(const FIoContainerHeader* ContainerHeader);
@@ -126,8 +131,9 @@ private:
 #endif //if WITH_EDITOR
 
 	void Update();
-	FEntryHandle AddNewEntryData(const FFilePackageStoreEntry& Entry);
+	FEntryHandle AddNewEntryData(const FFilePackageStoreEntry& Entry, const FFilePackageStoreEntrySoftReferences* SoftReferences);
 	FEntryHandle AddOldEntryData(FEntryHandle OldHandle, TConstArrayView<uint32> OldEntryData);
 	TConstArrayView<FPackageId> GetImportedPackages(FEntryHandle Handle);
 	TConstArrayView<FSHAHash> GetShaderHashes(FEntryHandle Handle);
+	TConstArrayView<uint32> GetSoftReferenceIndices(FEntryHandle Handle);
 };
diff --git a/Engine/Source/Runtime/PakFile/Private/IoDispatcherFileBackend.cpp b/Engine/Source/Runtime/PakFile/Private/IoDispatcherFileBackend.cpp
--- a/Engine/Source/Runtime/PakFile/Private/IoDispatcherFileBackend.cpp
+++ b/Engine/Source/Runtime/PakFile/Private/IoDispatcherFileBackend.cpp
@@ -933,6 +933,28 @@ TIoStatusOr<FIoContainerHeader> FFileIoStoreReader::ReadContainerHeader() const
 		UE_LOG(LogIoDispatcher, Warning, TEXT("Invalid container header in file '%s'"), *ContainerFile.FilePath);
 		ContainerHeader = FIoContainerHeader();
 	}
+
+	if (ContainerHeader.SoftPackageReferencesSerialInfo.Size > 0)
+	{
+		if (ContainerHeader.SoftPackageReferencesSerialInfo.Offset < 0)
+		{
+			FIoStatus Status = FIoStatusBuilder(EIoErrorCode::ReadError)
+				<< FString::Printf(TEXT("Invalid soft package reference offset '" INT64_FMT "'"), ContainerHeader.SoftPackageReferencesSerialInfo.Offset);
+			return Status;
+		}
+		if ((ContainerHeader.SoftPackageReferencesSerialInfo.Offset + ContainerHeader.SoftPackageReferencesSerialInfo.Size) > Ar.TotalSize())
+		{
+			FIoStatus Status = FIoStatusBuilder(EIoErrorCode::ReadError)
+				<< FString::Printf(TEXT("Soft package reference offset '" INT64_FMT "' and size '" INT64_FMT "' will seek past the end of archive size '" INT64_FMT "'"),
+					ContainerHeader.SoftPackageReferencesSerialInfo.Offset,
+					ContainerHeader.SoftPackageReferencesSerialInfo.Size,
+					Ar.TotalSize());
+			return Status;
+		}
+		Ar.Seek(ContainerHeader.SoftPackageReferencesSerialInfo.Offset);
+		Ar << ContainerHeader.SoftPackageReferences;
+	}
+
 	return ContainerHeader;
 }
 
