Allow conditionally loading iostore container soft references when mounting.
Add a new MountContainers core delegate that allows for passing mount options and also allows for getting errors back from mounting.
FPakPlatformFile::Mount can now return errors.
Add new pak file entry only if both pak and iostore succeed - it shouldn't be possible to be partially mounted.

diff --git a/Engine/Source/Runtime/Core/Private/Misc/CoreDelegates.cpp b/Engine/Source/Runtime/Core/Private/Misc/CoreDelegates.cpp
--- a/Engine/Source/Runtime/Core/Private/Misc/CoreDelegates.cpp
+++ b/Engine/Source/Runtime/Core/Private/Misc/CoreDelegates.cpp
@@ -29,6 +29,7 @@ TMulticastDelegate<void()>& FCoreDelegates::GetPreMainInitDelegate()
 
 TDelegate<int32(const TArray<FString>&)> FCoreDelegates::OnMountAllPakFiles;
 TDelegate<IPakFile*(const FString&, int32)> FCoreDelegates::MountPak;
+TDelegate<bool(TArrayView<UE::FMountPaksExArgs>)> FCoreDelegates::MountPaksEx;
 TDelegate<bool(const FString&)> FCoreDelegates::OnUnmountPak;
 TDelegate<int32(const FString&)> FCoreDelegates::PakChunkIdOverrideDelegate;
 
diff --git a/Engine/Source/Runtime/Core/Public/GenericPlatform/GenericPlatformFile.h b/Engine/Source/Runtime/Core/Public/GenericPlatform/GenericPlatformFile.h
--- a/Engine/Source/Runtime/Core/Public/GenericPlatform/GenericPlatformFile.h
+++ b/Engine/Source/Runtime/Core/Public/GenericPlatform/GenericPlatformFile.h
@@ -972,7 +972,7 @@ public:
 	//~ End IPlatformFile Interface
 };
 
-/* Interface class for FPakFile to allow usage from modules that cannot have a compile dependency on FPakFile */
+/* Interface class for FPakFile to allow usage from modules that cannot have a compile dependency on FPakFile. See FCoreDelegates */
 class IPakFile
 {
 public:
@@ -993,6 +993,20 @@ public:
 	virtual int32 GetNumFiles() const = 0;
 };
 
+/* Pak Mount Options allow usage from modules that cannot have a compile dependency on FPakFile. See FCoreDelegates. */
+struct FPakMountOptions
+{
+	enum class EMountFlags : uint32
+	{
+		None = 0,
+		/** Make soft references available */
+		WithSoftReferences = 1 << 0,
+	};
+
+	EMountFlags MountFlags = EMountFlags::None;
+};
+ENUM_CLASS_FLAGS(FPakMountOptions::EMountFlags);
+
 inline uint32 GetTypeHash(const FFileJournalFileHandle& A)
 {
 	constexpr uint32 Mult = 103;
diff --git a/Engine/Source/Runtime/Core/Public/Misc/CoreDelegates.h b/Engine/Source/Runtime/Core/Public/Misc/CoreDelegates.h
--- a/Engine/Source/Runtime/Core/Public/Misc/CoreDelegates.h
+++ b/Engine/Source/Runtime/Core/Public/Misc/CoreDelegates.h
@@ -20,6 +20,7 @@
 #include "Misc/AES.h"
 #include "Misc/Build.h"
 #include "Misc/Optional.h"
+#include "IO/IoStatus.h"
 #include "Templates/Function.h"
 #include "Templates/SharedPointer.h"
 
@@ -79,6 +80,21 @@ struct FMultiprocessDetachedContext
 	int32 Id;
 	bool bAbnormalDetach;
 };
+
+struct FMountPaksExArgs
+{
+	struct FMountResult
+	{
+		// Any pak file that was mounted, may be null on success if an encryption key is pending
+		IPakFile* PakFile = nullptr;
+	};
+
+	const TCHAR* PakFilePath = nullptr;
+	int32 Order = INDEX_NONE;
+	FPakMountOptions MountOptions;
+	// TODO: @Daniel.Lamb - this should use FError
+	TValueOrError<FMountResult, FString> Result = MakeValue();
+};
 }
 
 class FCoreDelegates
@@ -122,6 +138,10 @@ public:
 	// Callback to prompt the pak system to mount a pak file
 	static CORE_API TDelegate<IPakFile*(const FString&, int32)> MountPak;
 
+	// Callback to prompt the pak system to mount multiple pak files
+	// Returns true if any paks were mounted.
+	static CORE_API TDelegate<bool(TArrayView<UE::FMountPaksExArgs>)> MountPaksEx;
+
 	// Callback to prompt the pak system to unmount a pak file.
 	static CORE_API TDelegate<bool(const FString&)> OnUnmountPak;
 
diff --git a/Engine/Source/Runtime/Experimental/IoStore/OnDemand/Private/OnDemandIoStore.cpp b/Engine/Source/Runtime/Experimental/IoStore/OnDemand/Private/OnDemandIoStore.cpp
--- a/Engine/Source/Runtime/Experimental/IoStore/OnDemand/Private/OnDemandIoStore.cpp
+++ b/Engine/Source/Runtime/Experimental/IoStore/OnDemand/Private/OnDemandIoStore.cpp
@@ -129,7 +129,7 @@ static TIoStatusOr<FIoBuffer> DecodeChunk(const FOnDemandChunkInfo& ChunkInfo, F
 }
 
 ///////////////////////////////////////////////////////////////////////////////
-static TIoStatusOr<FSharedContainerHeader> DeserializeContainerHeader(const FOnDemandChunkInfo& ChunkInfo, FMemoryView EncodedHeaderChunk)
+static TIoStatusOr<FSharedContainerHeader> DeserializeContainerHeader(const FOnDemandChunkInfo& ChunkInfo, FMemoryView EncodedHeaderChunk, bool bReadSoftRefs)
 {
 	TIoStatusOr<FIoBuffer> Chunk = DecodeChunk(ChunkInfo, EncodedHeaderChunk);
 	if (Chunk.IsOk() == false)
@@ -140,7 +140,7 @@ static TIoStatusOr<FSharedContainerHeader> DeserializeContainerHeader(const FOnD
 	FSharedContainerHeader OutHeader = MakeShared<FIoContainerHeader>();
 	FMemoryReaderView Ar(Chunk.ValueOrDie().GetView());
 	Ar << *OutHeader;
-	if (OutHeader->SoftPackageReferencesSerialInfo.Size > 0)
+	if (bReadSoftRefs && OutHeader->SoftPackageReferencesSerialInfo.Size > 0)
 	{
 		if (OutHeader->SoftPackageReferencesSerialInfo.Offset < 0)
 		{
@@ -1316,6 +1316,11 @@ FIoStatus FOnDemandIoStore::TickMountRequest(FMountRequest& MountRequest)
 		else if (EnumHasAnyFlags(Args.Options, EOnDemandMountOptions::InstallOnDemand))
 		{
 			EnumAddFlags(Container->Flags, EOnDemandContainerFlags::InstallOnDemand);
+
+			if (EnumHasAnyFlags(Args.Options, EOnDemandMountOptions::WithSoftReferences))
+			{
+				EnumAddFlags(Container->Flags, EOnDemandContainerFlags::WithSoftReferences);
+			}
 		}
 
 		if (EnumHasAnyFlags(Container->Flags, EOnDemandContainerFlags::Encrypted) &&
@@ -1366,7 +1371,8 @@ FIoStatus FOnDemandIoStore::TickMountRequest(FMountRequest& MountRequest)
 
 				TIoStatusOr<FSharedContainerHeader> Header = Private::DeserializeContainerHeader(
 					ChunkInfo,
-					Response.ConsumeValueOrDie().Body.GetView());
+					Response.ConsumeValueOrDie().Body.GetView(),
+					EnumHasAnyFlags(Container->Flags, EOnDemandContainerFlags::WithSoftReferences));
 
 				if (Header.IsOk() == false)
 				{
@@ -1857,7 +1863,8 @@ void FOnDemandIoStore::OnEncryptionKeyAdded(const FGuid& Id, const FAES::FAESKey
 			{
 				TIoStatusOr<FSharedContainerHeader> Header = Private::DeserializeContainerHeader(
 					ChunkInfo,
-					EncodedHeader->GetView());
+					EncodedHeader->GetView(),
+					EnumHasAnyFlags(Container->Flags, EOnDemandContainerFlags::WithSoftReferences));
 
 				if (ensure(Header.IsOk()))
 				{
diff --git a/Engine/Source/Runtime/Experimental/IoStore/OnDemand/Private/OnDemandIoStore.h b/Engine/Source/Runtime/Experimental/IoStore/OnDemand/Private/OnDemandIoStore.h
--- a/Engine/Source/Runtime/Experimental/IoStore/OnDemand/Private/OnDemandIoStore.h
+++ b/Engine/Source/Runtime/Experimental/IoStore/OnDemand/Private/OnDemandIoStore.h
@@ -43,6 +43,7 @@ enum class EOnDemandContainerFlags : uint8
 	StreamOnDemand			= (1 << 2),
 	InstallOnDemand			= (1 << 3),
 	Encrypted				= (1 << 4),
+	WithSoftReferences		= (1 << 5),
 	Count
 };
 ENUM_CLASS_FLAGS(EOnDemandContainerFlags);
diff --git a/Engine/Source/Runtime/Experimental/IoStore/OnDemandCore/Public/IO/IoStoreOnDemand.h b/Engine/Source/Runtime/Experimental/IoStore/OnDemandCore/Public/IO/IoStoreOnDemand.h
--- a/Engine/Source/Runtime/Experimental/IoStore/OnDemandCore/Public/IO/IoStoreOnDemand.h
+++ b/Engine/Source/Runtime/Experimental/IoStore/OnDemandCore/Public/IO/IoStoreOnDemand.h
@@ -117,6 +117,8 @@ enum class EOnDemandMountOptions
 	InstallOnDemand			= 1 << 1,
 	/** Trigger callback on game thread. */
 	CallbackOnGameThread	= 1 << 2,
+	/** Make soft references available */
+	WithSoftReferences		= 1 << 3,
 };
 ENUM_CLASS_FLAGS(EOnDemandMountOptions);
 
diff --git a/Engine/Source/Runtime/PakFile/Private/FileIoDispatcherBackend.cpp b/Engine/Source/Runtime/PakFile/Private/FileIoDispatcherBackend.cpp
--- a/Engine/Source/Runtime/PakFile/Private/FileIoDispatcherBackend.cpp
+++ b/Engine/Source/Runtime/PakFile/Private/FileIoDispatcherBackend.cpp
@@ -641,7 +641,12 @@ public:
 	virtual									~FFileIoDispatcherBackend();
 
 	// IFileIoDispatcherBackend
-	virtual TIoStatusOr<FIoContainerHeader> Mount(const TCHAR* TocPath, int32 Order, const FGuid& EncryptionKeyGuid, const FAES::FAESKey& EncryptionKey) override;
+	virtual TIoStatusOr<FIoContainerHeader> Mount(
+		const TCHAR* TocPath, 
+		int32 Order, 
+		const FGuid& EncryptionKeyGuid, 
+		const FAES::FAESKey& EncryptionKey, 
+		UE::IoStore::ETocMountOptions Options) override;
 	virtual bool							Unmount(const TCHAR* TocPath) override;
 	virtual void							ReopenAllFileHandles() override;
 
@@ -676,7 +681,12 @@ FFileIoDispatcherBackend::~FFileIoDispatcherBackend()
 {
 }
 
-TIoStatusOr<FIoContainerHeader> FFileIoDispatcherBackend::Mount(const TCHAR* TocPath, int32 Order, const FGuid& EncryptionKeyGuid, const FAES::FAESKey& EncryptionKey)
+TIoStatusOr<FIoContainerHeader> FFileIoDispatcherBackend::Mount(
+	const TCHAR* TocPath, 
+	int32 Order, 
+	const FGuid& EncryptionKeyGuid, 
+	const FAES::FAESKey& EncryptionKey, 
+	UE::IoStore::ETocMountOptions Options)
 {
 	LLM_SCOPE_BYNAME(TEXT("FileSystem/FileIoStore"));
 	return IoStore.Mount(TocPath, Order);
diff --git a/Engine/Source/Runtime/PakFile/Private/FileIoDispatcherBackend.h b/Engine/Source/Runtime/PakFile/Private/FileIoDispatcherBackend.h
--- a/Engine/Source/Runtime/PakFile/Private/FileIoDispatcherBackend.h
+++ b/Engine/Source/Runtime/PakFile/Private/FileIoDispatcherBackend.h
@@ -5,19 +5,33 @@
 #include "IO/IoDispatcherBackend.h"
 #include "IO/IoStatus.h"
 #include "Misc/AES.h"
+#include "Misc/EnumClassFlags.h"
 
 struct FIoContainerHeader;
 
 namespace UE::IoStore
 {
 
+enum class ETocMountOptions
+{
+	None = 0,
+	/** Make soft references available */
+	WithSoftReferences = 1 << 0,
+};
+ENUM_CLASS_FLAGS(ETocMountOptions);
+
 class IFileIoDispatcherBackend
 	: public IIoDispatcherBackend
 {
 public:
 	virtual ~IFileIoDispatcherBackend() = default;
 
-	virtual TIoStatusOr<FIoContainerHeader> Mount(const TCHAR* TocPath, int32 Order, const FGuid& EncryptionKeyGuid, const FAES::FAESKey& EncryptionKey) = 0;
+	virtual TIoStatusOr<FIoContainerHeader> Mount(
+		const TCHAR* TocPath, 
+		int32 Order, 
+		const FGuid& EncryptionKeyGuid, 
+		const FAES::FAESKey& EncryptionKey, 
+		ETocMountOptions Options = ETocMountOptions::None) = 0;
 	virtual bool							Unmount(const TCHAR* TocPath) = 0;
 	virtual void							ReopenAllFileHandles() = 0;
 };
diff --git a/Engine/Source/Runtime/PakFile/Private/IPlatformFilePak.cpp b/Engine/Source/Runtime/PakFile/Private/IPlatformFilePak.cpp
--- a/Engine/Source/Runtime/PakFile/Private/IPlatformFilePak.cpp
+++ b/Engine/Source/Runtime/PakFile/Private/IPlatformFilePak.cpp
@@ -4505,7 +4505,16 @@ void FPakPlatformFile::HandleReloadPakReadersCommand(const TCHAR* Cmd, FOutputDe
 }
 #endif // !UE_BUILD_SHIPPING
 
-
+namespace UE::PakFile::Private
+{
+	/**
+	 * Handler for delegate to prompt us to load multiple containers.
+	 */
+	struct FHandleMountPaksExDelegate
+	{
+		static bool HandleDelegate(FPakPlatformFile* Self, TArrayView<UE::FMountPaksExArgs> MountArgs);
+	};
+}
 
 FPakPlatformFile::FPakPlatformFile()
 	: LowerLevel(NULL)
@@ -4525,6 +4534,7 @@ FPakPlatformFile::~FPakPlatformFile()
 
 	FCoreDelegates::OnMountAllPakFiles.Unbind();
 	FCoreDelegates::MountPak.Unbind();
+	FCoreDelegates::MountPaksEx.Unbind();
 	FCoreDelegates::OnUnmountPak.Unbind();
 	FCoreDelegates::OnOptimizeMemoryUsageForMountedPaks.Unbind();
 
@@ -5697,6 +5707,10 @@ bool FPakPlatformFile::Initialize(IPlatformFile* Inner, const TCHAR* CmdLine)
 
 	FCoreDelegates::OnMountAllPakFiles.BindRaw(this, &FPakPlatformFile::MountAllPakFiles);
 	FCoreDelegates::MountPak.BindRaw(this, &FPakPlatformFile::HandleMountPakDelegate);
+	FCoreDelegates::MountPaksEx.BindLambda([this](TArrayView<UE::FMountPaksExArgs> MountArgs)
+	{ 
+		return UE::PakFile::Private::FHandleMountPaksExDelegate::HandleDelegate(this, MountArgs);
+	});
 	FCoreDelegates::OnUnmountPak.BindRaw(this, &FPakPlatformFile::HandleUnmountPakDelegate);
 	FCoreDelegates::OnOptimizeMemoryUsageForMountedPaks.BindRaw(this, &FPakPlatformFile::OptimizeMemoryUsageForMountedPaks);
 	FCoreInternalDelegates::GetCurrentlyMountedPaksDelegate().BindLambda([this]()
@@ -5886,11 +5900,35 @@ void FPakPlatformFile::OptimizeMemoryUsageForMountedPaks()
 }
 
 
-bool FPakPlatformFile::Mount(const TCHAR* InPakFilename, uint32 PakOrder, const TCHAR* InPath /*= nullptr*/, bool bLoadIndex /*= true*/, FPakListEntry* OutPakListEntry /*= nullptr*/)
+bool FPakPlatformFile::Mount(
+	const TCHAR* InPakFilename, 
+	uint32 PakOrder, 
+	const TCHAR* InPath /*= nullptr*/, 
+	bool bLoadIndex /*= true*/, 
+	FPakListEntry* OutPakListEntry /*= nullptr*/)
+{
+	FPakMountArgs MountArgs
+	{
+		.PakFilename = InPakFilename,
+		.PakOrder = PakOrder,
+		.Path = InPath,
+		.bLoadIndex = bLoadIndex
+	};
+
+	return Mount(MountArgs, nullptr, OutPakListEntry);
+}
+
+bool FPakPlatformFile::Mount(const FPakMountArgs& MountArgs, FIoStatus* OutIoMountStatus /*= nullptr*/, FPakListEntry* OutPakListEntry /*= nullptr*/)
 {
 	LLM_SCOPE_BYNAME(TEXT("FileSystem/PakFile"));
+
+	const TCHAR* InPakFilename = MountArgs.PakFilename;
+	uint32 PakOrder = MountArgs.PakOrder;
+	const TCHAR* InPath = MountArgs.Path;
+	const bool bLoadIndex = MountArgs.bLoadIndex;
+
 	bool bPakSuccess = false;
-	bool bIoStoreSuccess = true;
+	FIoStatus IoStoreSuccess = FIoStatus::Unknown;
 	if (LowerLevel->FileExists(InPakFilename))
 	{
 		TRefCountPtr<FPakFile> Pak = MakeRefCount<FPakFile>(LowerLevel, InPakFilename, bSigned, bLoadIndex);
@@ -5930,21 +5968,8 @@ bool FPakPlatformFile::Mount(const TCHAR* InPakFilename, uint32 PakOrder, const
 					}
 					PakOrder += 100 * ChunkVersionNumber;
 				}
-				{
-					// Add new pak file
-					UE::TScopeLock ScopedLock(PakListCritical);
-					FPakListEntry Entry;
-					Entry.ReadOrder = PakOrder;
-					Entry.PakFile = Pak;
-					Pak->SetIsMounted(true);
-					PakFiles.Add(Entry);
-					PakFiles.StableSort();
 
-					if (OutPakListEntry)
-					{
-						*OutPakListEntry = MoveTemp(Entry);
-					}
-				}
+				// Pak will be added to list if iostore mount also succeeds
 				bPakSuccess = true;
 			}
 			else
@@ -5960,142 +5985,193 @@ bool FPakPlatformFile::Mount(const TCHAR* InPakFilename, uint32 PakOrder, const
 				Entry.PakchunkIndex = Pak->PakchunkIndex;
 
 				Pak.SafeRelease();
+
+				if (OutIoMountStatus)
+				{
+					*OutIoMountStatus = FIoStatus(EIoErrorCode::PendingEncryptionKey, TEXT("Encryption key not available"));
+				}
+
 				return false;
 			}
 		}
 		else
 		{
 			UE_LOG(LogPakFile, Warning, TEXT("Failed to mount pak \"%s\", pak is invalid."), InPakFilename);
+			IoStoreSuccess = FIoStatus(EIoErrorCode::SignatureError, TEXT("Pak is invalid"));
 		}
 
-		if (bPakSuccess && IoDispatcherFileBackend.IsValid())
+		if (!bPakSuccess)
 		{
-			FGuid EncryptionKeyGuid = Pak->GetInfo().EncryptionKeyGuid;
-			FAES::FAESKey EncryptionKey;
-
-			if (!UE::FEncryptionKeyManager::Get().TryGetKey(EncryptionKeyGuid, EncryptionKey))
+			Pak.SafeRelease();
+		}
+		else
+		{
+			if (!IoDispatcherFileBackend.IsValid())
 			{
-				if (!EncryptionKeyGuid.IsValid() && FCoreDelegates::GetPakEncryptionKeyDelegate().IsBound())
-				{
-					FCoreDelegates::GetPakEncryptionKeyDelegate().Execute(EncryptionKey.Key);
-				}
+				IoStoreSuccess = FIoStatus::Ok;
 			}
-
-			FString UtocPath = FPaths::ChangeExtension(InPakFilename, TEXT(".utoc"));
-			if (FPlatformFileManager::Get().GetPlatformFile().FileExists(*UtocPath))
+			else
 			{
-				if (ShouldCheckPak())
+				FGuid EncryptionKeyGuid = Pak->GetInfo().EncryptionKeyGuid;
+				FAES::FAESKey EncryptionKey;
+
+				if (!UE::FEncryptionKeyManager::Get().TryGetKey(EncryptionKeyGuid, EncryptionKey))
 				{
-					ensure(CheckIoStoreContainerBlockSignatures(*UtocPath));
+					if (!EncryptionKeyGuid.IsValid() && FCoreDelegates::GetPakEncryptionKeyDelegate().IsBound())
+					{
+						FCoreDelegates::GetPakEncryptionKeyDelegate().Execute(EncryptionKey.Key);
+					}
 				}
-				TIoStatusOr<FIoContainerHeader> MountResult = IoDispatcherFileBackend->Mount(*UtocPath, PakOrder, EncryptionKeyGuid, EncryptionKey);
-				if (MountResult.IsOk())
+
+				const FString UtocPath = FPaths::ChangeExtension(InPakFilename, TEXT(".utoc"));
+				if (FPlatformFileManager::Get().GetPlatformFile().FileExists(*UtocPath))
 				{
-					UE_LOG(LogPakFile, Display, TEXT("Mounted IoStore container \"%s\""), *UtocPath);
-					Pak->IoContainerHeader = MakeUnique<FIoContainerHeader>(MountResult.ConsumeValueOrDie());
-					PackageStoreBackend->Mount(Pak->IoContainerHeader.Get(), PakOrder);
-#if WITH_EDITOR
-					FString OptionalSegmentUtocPath = FPaths::ChangeExtension(InPakFilename, FString::Printf(TEXT("%s.utoc"), FPackagePath::GetOptionalSegmentExtensionModifier()));
-					if (FPlatformFileManager::Get().GetPlatformFile().FileExists(*OptionalSegmentUtocPath))
+					if (ShouldCheckPak())
 					{
-						MountResult = IoDispatcherFileBackend->Mount(*OptionalSegmentUtocPath, PakOrder, EncryptionKeyGuid, EncryptionKey);
-						if (MountResult.IsOk())
-						{
-							Pak->OptionalSegmentIoContainerHeader = MakeUnique<FIoContainerHeader>(MountResult.ConsumeValueOrDie());
-							PackageStoreBackend->Mount(Pak->OptionalSegmentIoContainerHeader.Get(), PakOrder);
-							UE_LOG(LogPakFile, Display, TEXT("Mounted optional segment extension IoStore container \"%s\""), *OptionalSegmentUtocPath);
-						}
-						else
+						ensure(CheckIoStoreContainerBlockSignatures(*UtocPath));
+					}
+
+					UE::IoStore::ETocMountOptions TocMountOptions = UE::IoStore::ETocMountOptions::None;
+					if (EnumHasAnyFlags(MountArgs.MountOptions.MountFlags, FPakMountOptions::EMountFlags::WithSoftReferences))
+					{
+						TocMountOptions |= UE::IoStore::ETocMountOptions::WithSoftReferences;
+					}
+
+					TIoStatusOr<FIoContainerHeader> MountResult = IoDispatcherFileBackend->Mount(
+						*UtocPath, PakOrder, EncryptionKeyGuid, EncryptionKey, TocMountOptions);
+					IoStoreSuccess = MountResult.Status();  // Extract status before consuming value. Consuming the value will reset the status
+
+					if (MountResult.IsOk())
+					{
+						UE_LOG(LogPakFile, Display, TEXT("Mounted IoStore container \"%s\""), *UtocPath);
+						Pak->IoContainerHeader = MakeUnique<FIoContainerHeader>(MountResult.ConsumeValueOrDie());
+						PackageStoreBackend->Mount(Pak->IoContainerHeader.Get(), PakOrder);
+#if WITH_EDITOR
+						FString OptionalSegmentUtocPath = FPaths::ChangeExtension(InPakFilename, FString::Printf(TEXT("%s.utoc"), FPackagePath::GetOptionalSegmentExtensionModifier()));
+						if (FPlatformFileManager::Get().GetPlatformFile().FileExists(*OptionalSegmentUtocPath))
 						{
-							UE_LOG(LogPakFile, Warning, TEXT("Failed to mount optional segment extension IoStore container \"%s\" [%s]"), *OptionalSegmentUtocPath, *MountResult.Status().ToString());
+							MountResult = IoDispatcherFileBackend->Mount(*OptionalSegmentUtocPath, PakOrder, EncryptionKeyGuid, EncryptionKey);
+							IoStoreSuccess = MountResult.Status();  // Extract status before consuming value. Consuming the value will reset the status
+
+							if (MountResult.IsOk())
+							{
+								Pak->OptionalSegmentIoContainerHeader = MakeUnique<FIoContainerHeader>(MountResult.ConsumeValueOrDie());
+								PackageStoreBackend->Mount(Pak->OptionalSegmentIoContainerHeader.Get(), PakOrder);
+								UE_LOG(LogPakFile, Display, TEXT("Mounted optional segment extension IoStore container \"%s\""), *OptionalSegmentUtocPath);
+							}
+							else
+							{
+								UE_LOG(LogPakFile, Warning, TEXT("Failed to mount optional segment extension IoStore container \"%s\" [%s]"), *OptionalSegmentUtocPath, *MountResult.Status().ToString());
+							}
 						}
-					}
 #endif
+					}
+					else
+					{
+						UE_LOG(LogPakFile, Warning, TEXT("Failed to mount IoStore container \"%s\" [%s]"), *UtocPath, *MountResult.Status().ToString());
+					}
 				}
 				else
 				{
-					bIoStoreSuccess = false;
-					UE_LOG(LogPakFile, Warning, TEXT("Failed to mount IoStore container \"%s\" [%s]"), *UtocPath, *MountResult.Status().ToString());
+					IoStoreSuccess = FIoStatus(EIoErrorCode::NotFound, TEXT("utoc not found"));
+					UE_LOG(LogPakFile, Warning, TEXT("IoStore container \"%s\" not found"), *UtocPath);
 				}
 			}
-			else
+
+			if (!IoStoreSuccess.IsOk())
 			{
-				bIoStoreSuccess = false;
-				UE_LOG(LogPakFile, Warning, TEXT("IoStore container \"%s\" not found"), *UtocPath);
+				Pak.SafeRelease();
 			}
-		}
+			else // Add new pak file only if both pak and iostore succeed
+			{
+				{
+					UE::TScopeLock ScopedLock(PakListCritical);
+					FPakListEntry Entry;
+					Entry.ReadOrder = PakOrder;
+					Entry.PakFile = Pak;
+					Pak->SetIsMounted(true);
+					PakFiles.Add(Entry);
+					PakFiles.StableSort();
 
-		if (bPakSuccess && FCoreInternalDelegates::GetOnPakMountOperation().IsBound())
-		{
-			FCoreInternalDelegates::GetOnPakMountOperation().Broadcast(EMountOperation::Mount, InPakFilename, PakOrder);
-		}
+					if (OutPakListEntry)
+					{
+						*OutPakListEntry = MoveTemp(Entry);
+					}
+				}
 
-		if (bPakSuccess)
-		{
-			double OnPakFileMounted2Time = 0.0;
-			{
-				FScopedDurationTimer Timer(OnPakFileMounted2Time);
+				if (FCoreInternalDelegates::GetOnPakMountOperation().IsBound())
+				{
+					FCoreInternalDelegates::GetOnPakMountOperation().Broadcast(EMountOperation::Mount, InPakFilename, PakOrder);
+				}
 
-PRAGMA_DISABLE_DEPRECATION_WARNINGS
-				if (FCoreDelegates::OnPakFileMounted2.IsBound())
+				double OnPakFileMounted2Time = 0.0;
 				{
-					// Avoid calling Broadcast if not in use; Broadcast even on an unsubscribed
-					// non-threadsafe delegate is not threadsafe.
-					FCoreDelegates::OnPakFileMounted2.Broadcast(*Pak);
+					FScopedDurationTimer Timer(OnPakFileMounted2Time);
+
+					PRAGMA_DISABLE_DEPRECATION_WARNINGS
+					if (FCoreDelegates::OnPakFileMounted2.IsBound())
+					{
+						// Avoid calling Broadcast if not in use; Broadcast even on an unsubscribed
+						// non-threadsafe delegate is not threadsafe.
+						FCoreDelegates::OnPakFileMounted2.Broadcast(*Pak);
+					}
+					PRAGMA_ENABLE_DEPRECATION_WARNINGS
+
+					FCoreDelegates::GetOnPakFileMounted2().Broadcast(*Pak);
 				}
-PRAGMA_ENABLE_DEPRECATION_WARNINGS
 
-				FCoreDelegates::GetOnPakFileMounted2().Broadcast(*Pak);
-			}
+				UE_LOG(LogPakFile, Display, TEXT("Mounted Pak file '%s', mount point: '%s'"), InPakFilename, *Pak->GetMountPoint());
+				UE_LOG(LogPakFile, Verbose, TEXT("OnPakFileMounted2Time == %lf"), OnPakFileMounted2Time);
 
-			UE_LOG(LogPakFile, Display, TEXT("Mounted Pak file '%s', mount point: '%s'"), InPakFilename, *Pak->GetMountPoint());
-			UE_LOG(LogPakFile, Verbose, TEXT("OnPakFileMounted2Time == %lf"), OnPakFileMounted2Time);
-						
-			// skip this check for the default mountpoint, it is a frequently used known-good mount point
-			FString NormalizedPakMountPoint = FPaths::CreateStandardFilename(Pak->GetMountPoint());
-			bool bIsMountingToRoot = NormalizedPakMountPoint == FPaths::CreateStandardFilename(FPaths::RootDir());
+				// skip this check for the default mountpoint, it is a frequently used known-good mount point
+				FString NormalizedPakMountPoint = FPaths::CreateStandardFilename(Pak->GetMountPoint());
+				bool bIsMountingToRoot = NormalizedPakMountPoint == FPaths::CreateStandardFilename(FPaths::RootDir());
 #if WITH_EDITOR
-			bIsMountingToRoot |= NormalizedPakMountPoint == FPaths::CreateStandardFilename(FPaths::GameFeatureRootPrefix());
+				bIsMountingToRoot |= NormalizedPakMountPoint == FPaths::CreateStandardFilename(FPaths::GameFeatureRootPrefix());
 #endif
-			if (!bIsMountingToRoot)
-			{
-				FString OutPackageName;
-				const FString& MountPoint = Pak->GetMountPoint();
-				if (!FPackageName::TryConvertFilenameToLongPackageName(MountPoint, OutPackageName))
+				if (!bIsMountingToRoot)
 				{
-					// Possibly the mount point is a parent path of mount points, e.g. <ProjectRoot>/Plugins,
-					// parent path of <ProjectRoot>/Plugins/PluginA and <ProjectRoot>/Plugins/PluginB.
-					// Do not warn in that case.
-					FString MountPointAbsPath = FPaths::ConvertRelativePathToFull(MountPoint);
-					bool bParentOfMountPoint = false;
-					for (const FString& ExistingMountPoint : FPackageName::QueryMountPointLocalAbsPaths())
+					FString OutPackageName;
+					const FString& MountPoint = Pak->GetMountPoint();
+					if (!FPackageName::TryConvertFilenameToLongPackageName(MountPoint, OutPackageName))
 					{
-						if (FPathViews::IsParentPathOf(MountPointAbsPath, ExistingMountPoint))
+						// Possibly the mount point is a parent path of mount points, e.g. <ProjectRoot>/Plugins,
+						// parent path of <ProjectRoot>/Plugins/PluginA and <ProjectRoot>/Plugins/PluginB.
+						// Do not warn in that case.
+						FString MountPointAbsPath = FPaths::ConvertRelativePathToFull(MountPoint);
+						bool bParentOfMountPoint = false;
+						for (const FString& ExistingMountPoint : FPackageName::QueryMountPointLocalAbsPaths())
 						{
-							bParentOfMountPoint = true;
-							break;
+							if (FPathViews::IsParentPathOf(MountPointAbsPath, ExistingMountPoint))
+							{
+								bParentOfMountPoint = true;
+								break;
+							}
+						}
+						if (!bParentOfMountPoint)
+						{
+							UE_LOG(LogPakFile, Display,
+								TEXT("Mount point '%s' is not mounted to a valid Root Path yet, ")
+								TEXT("assets in this pak file may not be accessible until a corresponding UFS Mount Point is added through FPackageName::RegisterMountPoint."),
+								*MountPoint);
 						}
-					}
-					if (!bParentOfMountPoint)
-					{
-						UE_LOG(LogPakFile, Display,
-							TEXT("Mount point '%s' is not mounted to a valid Root Path yet, ")
-							TEXT("assets in this pak file may not be accessible until a corresponding UFS Mount Point is added through FPackageName::RegisterMountPoint."),
-							*MountPoint);
 					}
 				}
 			}
 		}
-		else
-		{
-			Pak.SafeRelease();
-		}
 	}
 	else
 	{
 		UE_LOG(LogPakFile, Warning, TEXT("Failed to open pak \"%s\""), InPakFilename);
+		IoStoreSuccess = FIoStatus(EIoErrorCode::NotFound, TEXT("Pak not found"));
+	}
+
+	if (OutIoMountStatus)
+	{
+		*OutIoMountStatus = IoStoreSuccess;
 	}
-	return bPakSuccess && bIoStoreSuccess;
+
+	return bPakSuccess && IoStoreSuccess.IsOk();
 }
 
 bool FPakPlatformFile::Unmount(const TCHAR* InPakFilename)
@@ -6291,7 +6367,7 @@ int32 FPakPlatformFile::MountAllPakFiles(const TArray<FString>& PakFolders, cons
 	return NumPakFilesMounted;
 }
 
-int32 FPakPlatformFile::GetPakOrderFromPakFilePath(const FString& PakFilePath)
+int32 FPakPlatformFile::GetPakOrderFromPakFilePath(const FStringView PakFilePath)
 {
 	if (PakFilePath.StartsWith(FString::Printf(TEXT("%sPaks/%s-"), *FPaths::ProjectContentDir(), FApp::GetProjectName())))
 	{
@@ -6330,6 +6406,100 @@ IPakFile* FPakPlatformFile::HandleMountPakDelegate(const FString& PakFilePath, i
 	return nullptr;
 }
 
+namespace UE::PakFile::Private
+{
+	template<bool bInAllowDuplicateKeys /*= false*/>
+	struct FPakEntryFilenameSetKeyFuncs : public DefaultKeyFuncs<FPakPlatformFile::FPakListEntry, bInAllowDuplicateKeys>
+	{
+		typedef typename TCallTraits<FStringView>::ValueType KeyInitType;
+		typedef typename DefaultKeyFuncs<FPakPlatformFile::FPakListEntry, bInAllowDuplicateKeys>::ElementInitType ElementInitType;
+
+		[[nodiscard]] static FORCEINLINE KeyInitType GetSetKey(ElementInitType Element)
+		{
+			return MakeStringView(Element.PakFile->GetFilename());
+		}
+
+		[[nodiscard]] static FORCEINLINE bool Matches(KeyInitType A, KeyInitType B)
+		{
+			return DefaultKeyFuncs<FStringView, bInAllowDuplicateKeys>::Matches(A, B);
+		}
+
+		[[nodiscard]] static FORCEINLINE uint32 GetKeyHash(KeyInitType Key)
+		{
+			return DefaultKeyFuncs<FStringView, bInAllowDuplicateKeys>::GetKeyHash(Key);
+		}
+	};
+}
+
+bool UE::PakFile::Private::FHandleMountPaksExDelegate::HandleDelegate(
+	FPakPlatformFile* Self, 
+	TArrayView<UE::FMountPaksExArgs> MountArgs)
+{
+	bool bMountedAnything = false;
+
+	// Find already mounted paks
+	TSet<FPakPlatformFile::FPakListEntry, UE::PakFile::Private::FPakEntryFilenameSetKeyFuncs<>> MountedPaks;
+	if (MountArgs.Num() > 0)
+	{
+		UE::TScopeLock ScopedLock(Self->PakListCritical);
+		MountedPaks.Empty(Self->PakFiles.Num());
+		for (FPakPlatformFile::FPakListEntry& Entry : Self->PakFiles)
+		{
+			MountedPaks.Add(Entry);
+		}
+	}
+
+	for (UE::FMountPaksExArgs& Args : MountArgs)
+	{
+		const FStringView PakFilePathView(Args.PakFilePath);
+
+		if (const FPakPlatformFile::FPakListEntry* MountedPak = MountedPaks.FindByHash(GetTypeHash(PakFilePathView), Args.PakFilePath))
+		{
+			UE_LOG(LogPakFile, Warning, TEXT("Pak file %s already exists, skipping mount."), Args.PakFilePath);
+			Args.Result = MakeValue(MountedPak->PakFile.GetReference());			
+			continue;
+		}
+
+		FPlatformMisc::LowLevelOutputDebugStringf(TEXT("Mounting pak file: %s \n"), Args.PakFilePath);
+
+		int32 PakOrder = Args.Order;
+		if (PakOrder == INDEX_NONE)
+		{
+			PakOrder = Self->GetPakOrderFromPakFilePath(PakFilePathView);
+		}
+
+		FPakMountArgs PakMountArgs
+		{
+			.PakFilename = Args.PakFilePath,
+			.PakOrder = static_cast<uint32>(PakOrder),
+			.MountOptions = Args.MountOptions
+		};
+
+		FIoStatus IoMountStatus;
+		FPakPlatformFile::FPakListEntry Pak;
+		const bool bMountSuccess = Self->Mount(PakMountArgs, &IoMountStatus, &Pak);
+
+		bMountedAnything = bMountedAnything || bMountSuccess;
+
+		if (bMountSuccess)
+		{
+			Args.Result = MakeValue(Pak.PakFile.GetReference());
+		}
+		else if (IoMountStatus.GetErrorCode() == EIoErrorCode::PendingEncryptionKey)
+		{
+			// Don't consider this an error, but we have no pak to return
+			Args.Result = MakeValue(nullptr);
+		}
+		else
+		{
+			// TODO: @Daniel.Lamb use FError here
+			Args.Result = MakeError(IoMountStatus.ToString());
+		}
+	}
+
+	return bMountedAnything;
+}
+
 bool FPakPlatformFile::HandleUnmountPakDelegate(const FString& PakFilePath)
 {
 	FPlatformMisc::LowLevelOutputDebugStringf(TEXT("Unmounting pak file: %s \n"), *PakFilePath);
diff --git a/Engine/Source/Runtime/PakFile/Private/IoDispatcherFileBackend.cpp b/Engine/Source/Runtime/PakFile/Private/IoDispatcherFileBackend.cpp
--- a/Engine/Source/Runtime/PakFile/Private/IoDispatcherFileBackend.cpp
+++ b/Engine/Source/Runtime/PakFile/Private/IoDispatcherFileBackend.cpp
@@ -847,7 +847,7 @@ IMappedFileHandle* FFileIoStoreReader::GetMappedContainerFileHandle(uint64 TocOf
 	return new FMappedFileProxy(Partition.MappedFileHandle.Get(), Partition.FileSize);
 }
 
-TIoStatusOr<FIoContainerHeader> FFileIoStoreReader::ReadContainerHeader() const
+TIoStatusOr<FIoContainerHeader> FFileIoStoreReader::ReadContainerHeader(bool bReadSoftRefs) const
 {
 	LLM_SCOPE(ELLMTag::AsyncLoading);
 	TRACE_CPUPROFILER_EVENT_SCOPE(ReadContainerHeader);
@@ -934,7 +934,7 @@ TIoStatusOr<FIoContainerHeader> FFileIoStoreReader::ReadContainerHeader() const
 		ContainerHeader = FIoContainerHeader();
 	}
 
-	if (ContainerHeader.SoftPackageReferencesSerialInfo.Size > 0)
+	if (bReadSoftRefs && ContainerHeader.SoftPackageReferencesSerialInfo.Size > 0)
 	{
 		if (ContainerHeader.SoftPackageReferencesSerialInfo.Offset < 0)
 		{
@@ -1296,7 +1296,12 @@ void FFileIoStore::Shutdown()
 	FScheduler::Get().GetOversubscriptionLimitReachedEvent().Remove(OversubscriptionLimitReached);
 }
 
-TIoStatusOr<FIoContainerHeader> FFileIoStore::Mount(const TCHAR* InTocPath, int32 Order, const FGuid& EncryptionKeyGuid, const FAES::FAESKey& EncryptionKey)
+TIoStatusOr<FIoContainerHeader> FFileIoStore::Mount(
+	const TCHAR* InTocPath, 
+	int32 Order, 
+	const FGuid& EncryptionKeyGuid, 
+	const FAES::FAESKey& EncryptionKey,
+	UE::IoStore::ETocMountOptions Options)
 {
 	LLM_SCOPE_BYNAME(TEXT("FileSystem/FileIoStore"));
 	TUniquePtr<FFileIoStoreReader> Reader(new FFileIoStoreReader(*PlatformImpl, Stats));
@@ -1319,7 +1324,8 @@ TIoStatusOr<FIoContainerHeader> FFileIoStore::Mount(const TCHAR* InTocPath, int3
 		}
 	}
 
-	TIoStatusOr<FIoContainerHeader> ContainerHeaderReadResult = Reader->ReadContainerHeader();
+	TIoStatusOr<FIoContainerHeader> ContainerHeaderReadResult = Reader->ReadContainerHeader(
+		EnumHasAnyFlags(Options, UE::IoStore::ETocMountOptions::WithSoftReferences));
 	FIoContainerHeader ContainerHeader;
 	if (ContainerHeaderReadResult.IsOk())
 	{
diff --git a/Engine/Source/Runtime/PakFile/Private/IoDispatcherFileBackend.h b/Engine/Source/Runtime/PakFile/Private/IoDispatcherFileBackend.h
--- a/Engine/Source/Runtime/PakFile/Private/IoDispatcherFileBackend.h
+++ b/Engine/Source/Runtime/PakFile/Private/IoDispatcherFileBackend.h
@@ -49,7 +49,7 @@ public:
 	const FGuid& GetEncryptionKeyGuid() const { return ContainerFile.EncryptionKeyGuid; }
 	void SetEncryptionKey(const FAES::FAESKey& Key) { ContainerFile.EncryptionKey = Key; }
 	const FAES::FAESKey& GetEncryptionKey() const { return ContainerFile.EncryptionKey; }
-	TIoStatusOr<FIoContainerHeader> ReadContainerHeader() const;
+	TIoStatusOr<FIoContainerHeader> ReadContainerHeader(bool bReadSoftRefs) const;
 	void ReopenAllFileHandles();
 
 private:
@@ -112,7 +112,12 @@ public:
 	~FFileIoStore();
 	void Initialize(TSharedRef<const FIoDispatcherBackendContext> Context) override;
 	void Shutdown() override;
-	virtual TIoStatusOr<FIoContainerHeader> Mount(const TCHAR* InTocPath, int32 Order, const FGuid& EncryptionKeyGuid, const FAES::FAESKey& EncryptionKey) override;
+	virtual TIoStatusOr<FIoContainerHeader> Mount(
+		const TCHAR* InTocPath, 
+		int32 Order, 
+		const FGuid& EncryptionKeyGuid, 
+		const FAES::FAESKey& EncryptionKey, 
+		UE::IoStore::ETocMountOptions Options) override;
 	virtual bool Unmount(const TCHAR* InTocPath) override;
 	void ResolveIoRequests(FIoRequestList Requests, FIoRequestList& OutUnresolved) override;
 	void CancelIoRequest(FIoRequestImpl* Request) override;
diff --git a/Engine/Source/Runtime/PakFile/Public/IPlatformFilePak.h b/Engine/Source/Runtime/PakFile/Public/IPlatformFilePak.h
--- a/Engine/Source/Runtime/PakFile/Public/IPlatformFilePak.h
+++ b/Engine/Source/Runtime/PakFile/Public/IPlatformFilePak.h
@@ -19,7 +19,6 @@
 #include "HAL/PlatformCrt.h"
 #include "HAL/PlatformTime.h"
 #include "HAL/UnrealMemory.h"
-#include "IO/IoContainerId.h"
 #include "Logging/LogCategory.h"
 #include "Logging/LogMacros.h"
 #include "Math/BigInt.h"
@@ -57,11 +56,19 @@
 
 class FChunkCacheWorker;
 class FFileIoStore;
+class FIoStatus;
 class FFilePackageStoreBackend;
 class FOutputDevice;
 class IAsyncReadFileHandle;
 class IMappedFileHandle;
-namespace UE::PakFile::Private { class FPakFileDirectoryVisitorBase; }
+namespace UE::PakFile::Private 
+{ 
+	class FPakFileDirectoryVisitorBase;
+	struct FHandleMountPaksExDelegate;
+
+	template<bool bInAllowDuplicateKeys = false>
+	struct FPakEntryFilenameSetKeyFuncs;
+}
 namespace UE::IoStore { class IFileIoDispatcherBackend; }
 struct FIoContainerHeader;
 
@@ -1949,6 +1956,23 @@ private:
 	}
 };
 
+/**
+ * Args for pak mounting
+ */
+struct FPakMountArgs
+{
+	// Pak filename
+	const TCHAR* PakFilename = nullptr;
+	// Search order
+	uint32 PakOrder = 0;
+	// Path to mount the pak at
+	const TCHAR* Path = nullptr;
+	// Mount options
+	FPakMountOptions MountOptions;
+	// Flag to load Index
+	bool bLoadIndex = true;
+};
+
 /**
  * Platform file wrapper to be able to use pak files.
  **/
@@ -1970,6 +1994,9 @@ class FPakPlatformFile : public IPlatformFile
 		}
 	};
 
+	template<bool bInAllowDuplicateKeys>
+	friend struct UE::PakFile::Private::FPakEntryFilenameSetKeyFuncs;
+
 	struct FPakListDeferredEntry
 	{
 		FString Filename;
@@ -2049,13 +2076,18 @@ class FPakPlatformFile : public IPlatformFile
 	* Hardcode default load ordering of game main pak -> game content -> engine content -> saved dir
 	* would be better to make this config but not even the config system is initialized here so we can't do that
 	*/
-	PAKFILE_API static int32 GetPakOrderFromPakFilePath(const FString& PakFilePath);
+	PAKFILE_API static int32 GetPakOrderFromPakFilePath(const FStringView PakFilePath);
 
 	/**
-	 * Handler for device delegate to prompt us to load a new pak.	 
+	 * Handler for device delegate to prompt us to load a new pak.
 	 */
 	PAKFILE_API IPakFile* HandleMountPakDelegate(const FString& PakFilePath, int32 PakOrder);
 
+	/**
+	 * Handler for delegate to prompt us to load multiple containers.
+	 */
+	friend struct UE::PakFile::Private::FHandleMountPaksExDelegate;
+
 	/**
 	 * Handler for device delegate to prompt us to unload a pak.
 	 */
@@ -2226,11 +2258,26 @@ public:
 	/**
 	 * Mounts a pak file at the specified path.
 	 *
+	 * In most cases, this should not be called directly. Instead use the mounting delegates in FCoreDelegates
+	 * 
 	 * @param InPakFilename Pak filename.
 	 * @param InPath Path to mount the pak at.
+	 * @param OutPakListEntry for internal use
+	 * @param OutIoMountStatus Status of the pak or associated iostore mount
 	 */
 	PAKFILE_API bool Mount(const TCHAR* InPakFilename, uint32 PakOrder, const TCHAR* InPath = nullptr, bool bLoadIndex = true, FPakListEntry* OutPakListEntry = nullptr);
 
+	/**
+	 * Mounts a pak file at the specified path.
+	 *
+	 * In most cases, this should not be called directly. Instead use the mounting delegates in FCoreDelegates
+	 *
+	 * @param MountArgs mount arguments
+	 * @param OutIoMountStatus Status of the pak or associated iostore mount
+	 * @param OutPakListEntry for internal use
+	 */
+	PAKFILE_API bool Mount(const FPakMountArgs& MountArgs, FIoStatus* OutIoMountStatus = nullptr, FPakListEntry* OutPakListEntry = nullptr);
+
 	PAKFILE_API bool Unmount(const TCHAR* InPakFilename);
 
 	PAKFILE_API int32 MountAllPakFiles(const TArray<FString>& PakFolders);
