Fix Android Build (see 40945948)

diff --git a/Engine/Source/Runtime/PakFile/Private/IPlatformFilePak.cpp b/Engine/Source/Runtime/PakFile/Private/IPlatformFilePak.cpp
--- a/Engine/Source/Runtime/PakFile/Private/IPlatformFilePak.cpp
+++ b/Engine/Source/Runtime/PakFile/Private/IPlatformFilePak.cpp
@@ -6456,7 +6456,7 @@ bool UE::PakFile::Private::FHandleMountPaksExDelegate::HandleDelegate(
 		if (const FPakPlatformFile::FPakListEntry* MountedPak = MountedPaks.FindByHash(GetTypeHash(PakFilePathView), Args.PakFilePath))
 		{
 			UE_LOG(LogPakFile, Warning, TEXT("Pak file %s already exists, skipping mount."), Args.PakFilePath);
-			Args.Result = MakeValue(MountedPak->PakFile.GetReference());			
+			Args.Result = MakeValue(UE::FMountPaksExArgs::FMountResult{ MountedPak->PakFile.GetReference() });
 			continue;
 		}
 
@@ -6483,12 +6483,12 @@ bool UE::PakFile::Private::FHandleMountPaksExDelegate::HandleDelegate(
 
 		if (bMountSuccess)
 		{
-			Args.Result = MakeValue(Pak.PakFile.GetReference());
+			Args.Result = MakeValue(UE::FMountPaksExArgs::FMountResult{ Pak.PakFile.GetReference() });
 		}
 		else if (IoMountStatus.GetErrorCode() == EIoErrorCode::PendingEncryptionKey)
 		{
 			// Don't consider this an error, but we have no pak to return
-			Args.Result = MakeValue(nullptr);
+			Args.Result = MakeValue(UE::FMountPaksExArgs::FMountResult{ nullptr });
 		}
 		else
 		{
