Fix alignment check

diff --git a/Engine/Source/Runtime/PakFile/Private/IoDispatcherFileBackend.cpp b/Engine/Source/Runtime/PakFile/Private/IoDispatcherFileBackend.cpp
--- a/Engine/Source/Runtime/PakFile/Private/IoDispatcherFileBackend.cpp
+++ b/Engine/Source/Runtime/PakFile/Private/IoDispatcherFileBackend.cpp
@@ -1988,7 +1988,7 @@ TIoStatusOr<FIoMappedRegion> FFileIoStore::OpenMapped(const FIoChunkId& ChunkId,
 			IMappedFileRegion* MappedFileRegion = MappedFileHandle->MapRegion(BlockOffset + Options.GetOffset(), ResolvedSize);
 			if (MappedFileRegion != nullptr)
 			{
-				check(IsAligned(MappedFileRegion->GetMappedPtr(), FPlatformProperties::GetMemoryMappingAlignment()));
+				check(IsAligned(MappedFileRegion->GetMappedPtr(), FPlatformMemory::GetStats().PageSize));
 				return FIoMappedRegion{ MappedFileHandle, MappedFileRegion };
 			}
 			else
