Fixed bug in FilePackageStore not processing soft references for some packages due to IsEmpty not taking soft references into account.

#rnx

diff --git a/Engine/Source/Runtime/PakFile/Private/FilePackageStore.cpp b/Engine/Source/Runtime/PakFile/Private/FilePackageStore.cpp
--- a/Engine/Source/Runtime/PakFile/Private/FilePackageStore.cpp
+++ b/Engine/Source/Runtime/PakFile/Private/FilePackageStore.cpp
@@ -303,14 +303,14 @@ static FStringView GetRootPathFromPackageName(const FStringView Path)
 
 static bool IsEmpty(FEntryHandle Handle) 
 {
-	return Handle.HasPackageIds + Handle.HasShaderMaps == 0;
+	return Handle.HasPackageIds + Handle.HasShaderMaps + Handle.HasSoftRefs == 0;
 }
 
 } // namespace UE::FilePackageStorePrivate
 
-static bool IsEmpty(const FFilePackageStoreEntry& E)
+static bool IsEmpty(const FFilePackageStoreEntry& E, const FFilePackageStoreEntrySoftReferences* SoftRefs)
 {
-	return E.ImportedPackages.Num() + E.ShaderMapHashes.Num() == 0;
+	return (E.ImportedPackages.Num() + E.ShaderMapHashes.Num() == 0) && (SoftRefs == nullptr || SoftRefs->Indices.Num() == 0);
 }
 
 ////////////////////////////////////////////////////////////////////////////////
@@ -454,7 +454,7 @@ FFilePackageStoreBackend::FEntryHandle FFilePackageStoreBackend::AddNewEntryData
 	FEntryHandle Out;
 	Out.Offset = EntryData.Num();
 
-	if (IsEmpty(Entry))
+	if (IsEmpty(Entry, SoftReferences))
 	{
 		EntryData.Add(0u);
 	}
