- Add CVar to configure mount failure on missing utoc

#rnx

diff --git a/Engine/Source/Runtime/PakFile/Private/IPlatformFilePak.cpp b/Engine/Source/Runtime/PakFile/Private/IPlatformFilePak.cpp
--- a/Engine/Source/Runtime/PakFile/Private/IPlatformFilePak.cpp
+++ b/Engine/Source/Runtime/PakFile/Private/IPlatformFilePak.cpp
@@ -5928,6 +5928,13 @@ bool FPakPlatformFile::Mount(
 	return Mount(MountArgs, nullptr, OutPakListEntry);
 }
 
+static bool GPakMountFailOnMissingUtoc = true;
+static FAutoConsoleVariableRef CVarPakMountFailOnMissingUtoc(
+	TEXT("pak.MountFailOnMissingUtoc"),
+	GPakMountFailOnMissingUtoc,
+	TEXT("If true, then missing utoc will generate EIOErrorCode::NotFound")
+);
+
 bool FPakPlatformFile::Mount(const FPakMountArgs& MountArgs, FIoStatus* OutIoMountStatus /*= nullptr*/, FPakListEntry* OutPakListEntry /*= nullptr*/)
 {
 	LLM_SCOPE_BYNAME(TEXT("FileSystem/PakFile"));
@@ -6083,8 +6090,15 @@ bool FPakPlatformFile::Mount(const FPakMountArgs& MountArgs, FIoStatus* OutIoMou
 				}
 				else
 				{
-					IoStoreSuccess = FIoStatus(EIoErrorCode::NotFound, TEXT("utoc not found"));
-					UE_LOG(LogPakFile, Warning, TEXT("IoStore container \"%s\" not found"), *UtocPath);
+					if (GPakMountFailOnMissingUtoc)
+					{
+						IoStoreSuccess = FIoStatus(EIoErrorCode::NotFound, TEXT("utoc not found"));
+						UE_LOG(LogPakFile, Warning, TEXT("IoStore container \"%s\" not found"), *UtocPath);
+					}
+					else
+					{
+						IoStoreSuccess = FIoStatus::Ok;
+					}
 				}
 			}
 
