- Switch from CVar mechanism to static bool, as the Pak mounting occurs way before CVars are brought up.

#rnx

diff --git a/Engine/Source/Runtime/PakFile/Private/IPlatformFilePak.cpp b/Engine/Source/Runtime/PakFile/Private/IPlatformFilePak.cpp
--- a/Engine/Source/Runtime/PakFile/Private/IPlatformFilePak.cpp
+++ b/Engine/Source/Runtime/PakFile/Private/IPlatformFilePak.cpp
@@ -67,10 +67,7 @@ static int64 GTotalLoaded = 0;
 static int64 GTotalLoadedLastTick = 0;
 #endif
 
-
-
-
-
+bool FPakPlatformFile::bMountFailOnMissingUtoc = true;
 
 #ifndef DISABLE_NONUFS_INI_WHEN_COOKED
 #define DISABLE_NONUFS_INI_WHEN_COOKED 0
@@ -5928,13 +5925,6 @@ bool FPakPlatformFile::Mount(
 	return Mount(MountArgs, nullptr, OutPakListEntry);
 }
 
-static bool GPakMountFailOnMissingUtoc = true;
-static FAutoConsoleVariableRef CVarPakMountFailOnMissingUtoc(
-	TEXT("pak.MountFailOnMissingUtoc"),
-	GPakMountFailOnMissingUtoc,
-	TEXT("If true, then missing utoc will generate EIOErrorCode::NotFound")
-);
-
 bool FPakPlatformFile::Mount(const FPakMountArgs& MountArgs, FIoStatus* OutIoMountStatus /*= nullptr*/, FPakListEntry* OutPakListEntry /*= nullptr*/)
 {
 	LLM_SCOPE_BYNAME(TEXT("FileSystem/PakFile"));
@@ -6090,7 +6080,7 @@ bool FPakPlatformFile::Mount(const FPakMountArgs& MountArgs, FIoStatus* OutIoMou
 				}
 				else
 				{
-					if (GPakMountFailOnMissingUtoc)
+					if (FPakPlatformFile::bMountFailOnMissingUtoc)
 					{
 						IoStoreSuccess = FIoStatus(EIoErrorCode::NotFound, TEXT("utoc not found"));
 						UE_LOG(LogPakFile, Warning, TEXT("IoStore container \"%s\" not found"), *UtocPath);
diff --git a/Engine/Source/Runtime/PakFile/Public/IPlatformFilePak.h b/Engine/Source/Runtime/PakFile/Public/IPlatformFilePak.h
--- a/Engine/Source/Runtime/PakFile/Public/IPlatformFilePak.h
+++ b/Engine/Source/Runtime/PakFile/Public/IPlatformFilePak.h
@@ -2134,6 +2134,8 @@ class FPakPlatformFile : public IPlatformFile
 
 public:
 
+	static PAKFILE_API bool bMountFailOnMissingUtoc;
+
 	//~ For visibility of overloads we don't override
 	using IPlatformFile::IterateDirectory;
 	using IPlatformFile::IterateDirectoryRecursively;
