Add ability to get currently mounted chunks. Useful to compare against asset registry chunks

diff --git a/Engine/Source/Runtime/PakFile/Public/IPlatformFilePak.h b/Engine/Source/Runtime/PakFile/Public/IPlatformFilePak.h
--- a/Engine/Source/Runtime/PakFile/Public/IPlatformFilePak.h
+++ b/Engine/Source/Runtime/PakFile/Public/IPlatformFilePak.h
@@ -2196,6 +2196,19 @@ public:
 		}
 	}
 
+	/**
+	* Gets a set of the mounted chunk Ids.
+	*/
+	FORCEINLINE void GetMountedChunkIds(TSet<int32>& OutChunkIds)
+	{
+		UE::TScopeLock ScopedLock(PakListCritical);
+		OutChunkIds.Empty(PakFiles.Num());
+		for (FPakListEntry& Entry : PakFiles)
+		{
+			OutChunkIds.Add(Entry.PakFile->PakGetPakchunkIndex());
+		}
+	}
+
 	/**
 	 * Checks if pak files exist in any of the known pak file locations.
 	 */
