FSignedArchiveReader: Fix deadlock in edgecase ordering of completed requests in FChunkCacheWorker::ProcessQueue().
#rnx

diff --git a/Engine/Source/Runtime/PakFile/Private/SignedArchiveReader.cpp b/Engine/Source/Runtime/PakFile/Private/SignedArchiveReader.cpp
--- a/Engine/Source/Runtime/PakFile/Private/SignedArchiveReader.cpp
+++ b/Engine/Source/Runtime/PakFile/Private/SignedArchiveReader.cpp
@@ -180,6 +180,7 @@ int32 FChunkCacheWorker::ProcessQueue()
 		INC_DWORD_STAT(STAT_FChunkCacheWorker_NumProcessQueueWithWork);
 	}
 
+	bool bFreedUpResourceWithoutHandingItOut = false;
 	for (int32 RequestIndex = 0; RequestIndex < ActiveRequests.Num(); ++RequestIndex)
 	{
 		FChunkRequest& Request = *ActiveRequests[RequestIndex];
@@ -190,6 +191,10 @@ int32 FChunkCacheWorker::ProcessQueue()
 			ReleaseBuffer(Request.Index);
 			ActiveRequests.RemoveAt(RequestIndex--);
 			FreeChunkRequests.Push(&Request);
+
+			// If we free a resource and don't hand it out in any follow-up requests after we free it, then trigger
+			// our event to run ProcessQueue again and to see if any earlier requests could use it.
+			bFreedUpResourceWithoutHandingItOut = true;
 		}
 		else if (Request.Buffer == NULL)
 		{
@@ -202,6 +207,7 @@ int32 FChunkCacheWorker::ProcessQueue()
 				if (!!CachedBuffer)
 				{
 					// Load and verify.
+					bFreedUpResourceWithoutHandingItOut = false;
 					CachedBuffer->ChunkIndex = Request.Index;
 					Request.Buffer = CachedBuffer;
 					CheckSignature(Request);
@@ -227,6 +233,10 @@ int32 FChunkCacheWorker::ProcessQueue()
 			}
 		}
 	}
+	if (bFreedUpResourceWithoutHandingItOut && QueuedRequestsEvent)
+	{
+		QueuedRequestsEvent->Trigger();
+	}
 
 	// Release any pending events
 	if (!EventsToRelease.IsEmpty())
