FFilePackageStoreBackend - Skip localized asset redirection in Editor
- This fixes crashes if redirection does not happen on game thread and the editor doesn't  normally do this redirection anyway.

diff --git a/Engine/Source/Runtime/PakFile/Private/FilePackageStore.cpp b/Engine/Source/Runtime/PakFile/Private/FilePackageStore.cpp
--- a/Engine/Source/Runtime/PakFile/Private/FilePackageStore.cpp
+++ b/Engine/Source/Runtime/PakFile/Private/FilePackageStore.cpp
@@ -366,19 +366,22 @@ bool FFilePackageStoreBackend::GetPackageRedirectInfo(FPackageId PackageId, FNam
 		return true;
 	}
 	
-	const FName* FindLocalizedPackageSourceName = LocalizedPackages.Find(PackageId);
-	if (FindLocalizedPackageSourceName)
+	if (!GIsEditor) // Skip localized asset redirection in Editor
 	{
-		FName LocalizedPackageName = FPackageLocalizationManager::Get().FindLocalizedPackageName(*FindLocalizedPackageSourceName);
-		if (!LocalizedPackageName.IsNone())
+		const FName* FindLocalizedPackageSourceName = LocalizedPackages.Find(PackageId);
+		if (FindLocalizedPackageSourceName)
 		{
-			FPackageId LocalizedPackageId = FPackageId::FromName(LocalizedPackageName);
-			if (PackageEntries.Find(LocalizedPackageId))
+			FName LocalizedPackageName = FPackageLocalizationManager::Get().FindLocalizedPackageName(*FindLocalizedPackageSourceName);
+			if (!LocalizedPackageName.IsNone())
 			{
-				OutSourcePackageName = *FindLocalizedPackageSourceName;
-				OutRedirectedToPackageId = LocalizedPackageId;
-				UE_LOG(LogFilePackageStore, Verbose, TEXT("Redirecting from localized package %s to 0x%s"), *OutSourcePackageName.ToString(), *LexToString(OutRedirectedToPackageId));
-				return true;
+				FPackageId LocalizedPackageId = FPackageId::FromName(LocalizedPackageName);
+				if (PackageEntries.Find(LocalizedPackageId))
+				{
+					OutSourcePackageName = *FindLocalizedPackageSourceName;
+					OutRedirectedToPackageId = LocalizedPackageId;
+					UE_LOG(LogFilePackageStore, Verbose, TEXT("Redirecting from localized package %s to 0x%s"), *OutSourcePackageName.ToString(), *LexToString(OutRedirectedToPackageId));
+					return true;
+				}
 			}
 		}
 	}
