[Bugfix] Client crash when loading corrupted pak files

Fixed crash when loading content from corrupted pak or optional segment files.

diff --git a/Engine/Source/Runtime/PakFile/Private/IPlatformFilePak.cpp b/Engine/Source/Runtime/PakFile/Private/IPlatformFilePak.cpp
--- a/Engine/Source/Runtime/PakFile/Private/IPlatformFilePak.cpp
+++ b/Engine/Source/Runtime/PakFile/Private/IPlatformFilePak.cpp
@@ -6050,9 +6050,7 @@ bool FPakPlatformFile::Mount(const FPakMountArgs& MountArgs, FIoStatus* OutIoMou
 
 					if (MountResult.IsOk())
 					{
-						UE_LOG(LogPakFile, Display, TEXT("Mounted IoStore container \"%s\""), *UtocPath);
 						Pak->IoContainerHeader = MakeUnique<FIoContainerHeader>(MountResult.ConsumeValueOrDie());
-						PackageStoreBackend->Mount(Pak->IoContainerHeader.Get(), PakOrder);
 #if WITH_EDITOR
 						FString OptionalSegmentUtocPath = FPaths::ChangeExtension(InPakFilename, FString::Printf(TEXT("%s.utoc"), FPackagePath::GetOptionalSegmentExtensionModifier()));
 						if (FPlatformFileManager::Get().GetPlatformFile().FileExists(*OptionalSegmentUtocPath))
@@ -6062,6 +6060,8 @@ bool FPakPlatformFile::Mount(const FPakMountArgs& MountArgs, FIoStatus* OutIoMou
 
 							if (MountResult.IsOk())
 							{
+								PackageStoreBackend->Mount(Pak->IoContainerHeader.Get(), PakOrder);
+								UE_LOG(LogPakFile, Display, TEXT("Mounted IoStore container \"%s\""), *UtocPath);
 								Pak->OptionalSegmentIoContainerHeader = MakeUnique<FIoContainerHeader>(MountResult.ConsumeValueOrDie());
 								PackageStoreBackend->Mount(Pak->OptionalSegmentIoContainerHeader.Get(), PakOrder);
 								UE_LOG(LogPakFile, Display, TEXT("Mounted optional segment extension IoStore container \"%s\""), *OptionalSegmentUtocPath);
@@ -6071,6 +6071,9 @@ bool FPakPlatformFile::Mount(const FPakMountArgs& MountArgs, FIoStatus* OutIoMou
 								UE_LOG(LogPakFile, Warning, TEXT("Failed to mount optional segment extension IoStore container \"%s\" [%s]"), *OptionalSegmentUtocPath, *MountResult.Status().ToString());
 							}
 						}
+#else
+						PackageStoreBackend->Mount(Pak->IoContainerHeader.Get(), PakOrder);
+						UE_LOG(LogPakFile, Display, TEXT("Mounted IoStore container \"%s\""), *UtocPath);
 #endif
 					}
 					else
