Improved error for missing softreference information from container.

diff --git a/Engine/Source/Runtime/Core/Private/IO/IoContainerHeader.cpp b/Engine/Source/Runtime/Core/Private/IO/IoContainerHeader.cpp
--- a/Engine/Source/Runtime/Core/Private/IO/IoContainerHeader.cpp
+++ b/Engine/Source/Runtime/Core/Private/IO/IoContainerHeader.cpp
@@ -34,6 +34,7 @@ FArchive& operator<<(FArchive& Ar, FIoContainerHeaderSoftPackageReferences& Soft
 	{
 		Ar << SoftPackageReferences.PackageIds;
 		Ar << SoftPackageReferences.PackageIndices;
+		SoftPackageReferences.bLoadedSoftPackageReferences = true;
 	}
 	return Ar;
 }
diff --git a/Engine/Source/Runtime/Core/Public/IO/IoContainerHeader.h b/Engine/Source/Runtime/Core/Public/IO/IoContainerHeader.h
--- a/Engine/Source/Runtime/Core/Public/IO/IoContainerHeader.h
+++ b/Engine/Source/Runtime/Core/Public/IO/IoContainerHeader.h
@@ -79,6 +79,9 @@ struct FIoContainerHeaderSoftPackageReferences
 	// Flag indicating whether any soft package references exists.
 	bool bContainsSoftPackageReferences = false;
 
+	// Transient flag indicating that the soft package references were loaded at runtime
+	bool bLoadedSoftPackageReferences = false;
+
 	CORE_API friend FArchive& operator<<(FArchive& Ar, FIoContainerHeaderSoftPackageReferences& SoftPackageReferences);
 };
 
diff --git a/Engine/Source/Runtime/CoreUObject/Private/Serialization/AsyncLoading2.cpp b/Engine/Source/Runtime/CoreUObject/Private/Serialization/AsyncLoading2.cpp
--- a/Engine/Source/Runtime/CoreUObject/Private/Serialization/AsyncLoading2.cpp
+++ b/Engine/Source/Runtime/CoreUObject/Private/Serialization/AsyncLoading2.cpp
@@ -6654,7 +6654,7 @@ void FAsyncPackage2::ImportPackagesRecursiveInner(FAsyncLoadingThreadState2& Thr
 			bIsFullyLoaded = false;
 		}
 #endif
-		const bool bIsNotInstalled = GFailLoadOnNotInstalledImport && ImportedPackageStatus == EPackageStoreEntryStatus::NotInstalled;
+		const bool bIsNotInstalled = GFailLoadOnNotInstalledImport && (ImportedPackageStatus == EPackageStoreEntryStatus::NotInstalled);
 		const bool bIsMissing = (ImportedPackageStatus == EPackageStoreEntryStatus::Missing) || (!GFailLoadOnNotInstalledImport && ImportedPackageStatus == EPackageStoreEntryStatus::NotInstalled);
 		if (bIsFullyLoaded)
 		{
diff --git a/Engine/Source/Runtime/Experimental/IoStore/OnDemand/Private/OnDemandContentInstaller.cpp b/Engine/Source/Runtime/Experimental/IoStore/OnDemand/Private/OnDemandContentInstaller.cpp
--- a/Engine/Source/Runtime/Experimental/IoStore/OnDemand/Private/OnDemandContentInstaller.cpp
+++ b/Engine/Source/Runtime/Experimental/IoStore/OnDemand/Private/OnDemandContentInstaller.cpp
@@ -33,6 +33,13 @@ namespace UE::IoStore
 			IoStoreErrorOnRequest,
 			TEXT("When the request with a debug name partially matching this cvar is found iostore will error with a random error.")
 			);
+
+		FString IoStoreDebugRequest = "";
+		static FAutoConsoleVariableRef CVar_IoStoreDebugRequest(
+			TEXT("iostore.DebugRequest"),
+			IoStoreDebugRequest,
+			TEXT("When the request with a debug name partially matching this cvar is found iostore will error with a random error.")
+		);
 #endif
 	}
 namespace Private
@@ -568,6 +575,15 @@ void FOnDemandContentInstaller::ProcessInstallRequest(FRequest& Request)
 			return CompleteInstallRequest(Request);
 		}
 	}
+
+	if (!UE::IoStore::CVars::IoStoreDebugRequest.IsEmpty())
+	{
+		if (FCString::Strstr(*ContentHandle->DebugName, *UE::IoStore::CVars::IoStoreDebugRequest) != nullptr ||
+			FCString::Strstr(*InstallRequest.Args.DebugName, *UE::IoStore::CVars::IoStoreDebugRequest) != nullptr)
+		{
+			UE_LOG(LogIoStoreOnDemand, Log, TEXT("Place breakpoint here to debug FOnDemandContentInstaller::ProcessInstallRequest for asset %s"), *ContentHandle->DebugName);
+		}
+	}
 #endif
 
 	if (Request.IsCancelled())
diff --git a/Engine/Source/Runtime/PakFile/Private/FilePackageStore.cpp b/Engine/Source/Runtime/PakFile/Private/FilePackageStore.cpp
--- a/Engine/Source/Runtime/PakFile/Private/FilePackageStore.cpp
+++ b/Engine/Source/Runtime/PakFile/Private/FilePackageStore.cpp
@@ -390,20 +390,46 @@ bool FFilePackageStoreBackend::GetPackageRedirectInfo(FPackageId PackageId, FNam
 
 TConstArrayView<uint32> FFilePackageStoreBackend::GetSoftReferences(FPackageId PackageId, TConstArrayView<FPackageId>& OutPackageIds)
 {
-	if (const FEntryHandle* Entry = PackageEntries.Find(PackageId); Entry != nullptr && Entry->HasSoftRefs)
+	if (const FEntryHandle* Entry = PackageEntries.Find(PackageId); Entry != nullptr)
 	{
-		FMountedContainer* Container = Algo::FindByPredicate(
-			MountedContainers,
-			[Entry](const FMountedContainer& C) { return C.EntryDataRange.Contains(Entry->Offset); });
-
-		if (Container != nullptr)
+		if (Entry->HasSoftRefs)
 		{
-			OutPackageIds = Container->ContainerHeader->SoftPackageReferences.PackageIds;
-			return GetSoftReferenceIndices(*Entry);
+			FMountedContainer* Container = Algo::FindByPredicate(
+				MountedContainers,
+				[Entry](const FMountedContainer& C) { return C.EntryDataRange.Contains(Entry->Offset); });
+
+			if (Container != nullptr)
+			{
+				// to do: figure out a better way to make sure softreference information wasn't loaded for this container, also this should make the outer request fail completely and error with appropriate information
+				if (Container->ContainerHeader->SoftPackageReferences.bLoadedSoftPackageReferences)
+				{
+					UE_LOG(LogFilePackageStore, Error, TEXT("Attempting to read softreference information for package %s but container %s never loaded softreference information"), *LexToString(PackageId), *LexToString(Container->ContainerHeader->ContainerId));
+				}
+				OutPackageIds = Container->ContainerHeader->SoftPackageReferences.PackageIds;
+				return GetSoftReferenceIndices(*Entry);
+			}
+			else
+			{
+				UE_LOG(LogFilePackageStore, Warning, TEXT("Failed to find container with soft references for package ID '%s'"), *LexToString(PackageId));
+			}
 		}
 		else
 		{
-			UE_LOG(LogFilePackageStore, Warning, TEXT("Failed to find container with soft references for package ID '%s'"), *LexToString(PackageId));
+			// make sure tha tit's containers mounted soft refs
+#if !UE_BUILD_SHIPPING
+			FMountedContainer* Container = Algo::FindByPredicate(
+				MountedContainers,
+				[Entry](const FMountedContainer& C) { return C.EntryDataRange.Contains(Entry->Offset); });
+
+			if (Container != nullptr)
+			{
+				// to do: figure out a better way to make sure softreference information wasn't loaded for this container, also this should make the outer request fail completely and error with appropriate information
+				if (!Container->ContainerHeader->SoftPackageReferences.bLoadedSoftPackageReferences)
+				{
+					UE_LOG(LogFilePackageStore, Error, TEXT("Attempting to read softreference information for package %s but container %s never loaded softreference information"), *LexToString(PackageId), *LexToString(Container->ContainerHeader->ContainerId));
+				}
+			}
+#endif
 		}
 	}
 
diff --git a/Engine/Source/Runtime/PakFile/Private/IoDispatcherFileBackend.cpp b/Engine/Source/Runtime/PakFile/Private/IoDispatcherFileBackend.cpp
--- a/Engine/Source/Runtime/PakFile/Private/IoDispatcherFileBackend.cpp
+++ b/Engine/Source/Runtime/PakFile/Private/IoDispatcherFileBackend.cpp
@@ -934,7 +934,8 @@ TIoStatusOr<FIoContainerHeader> FFileIoStoreReader::ReadContainerHeader(bool bRe
 		ContainerHeader = FIoContainerHeader();
 	}
 
-	if (bReadSoftRefs && ContainerHeader.SoftPackageReferencesSerialInfo.Size > 0)
+	// todo, ensure here if the container didn't include softreferences and we want to load them, the size check isn't really sufficent
+	if (bReadSoftRefs && ensureMsgf(ContainerHeader.SoftPackageReferencesSerialInfo.Size > 4, TEXT("Need to cook with softreferences enabled")))
 	{
 		if (ContainerHeader.SoftPackageReferencesSerialInfo.Offset < 0)
 		{
