Removed overzelous ensure for mounting pak files with softreferences.

diff --git a/Engine/Source/Runtime/PakFile/Private/IoDispatcherFileBackend.cpp b/Engine/Source/Runtime/PakFile/Private/IoDispatcherFileBackend.cpp
--- a/Engine/Source/Runtime/PakFile/Private/IoDispatcherFileBackend.cpp
+++ b/Engine/Source/Runtime/PakFile/Private/IoDispatcherFileBackend.cpp
@@ -935,7 +935,7 @@ TIoStatusOr<FIoContainerHeader> FFileIoStoreReader::ReadContainerHeader(bool bRe
 	}
 
 	// todo, ensure here if the container didn't include softreferences and we want to load them, the size check isn't really sufficent
-	if (bReadSoftRefs && ensureMsgf(ContainerHeader.SoftPackageReferencesSerialInfo.Size > 4, TEXT("Need to cook with softreferences enabled")))
+	if (bReadSoftRefs)
 	{
 		if (ContainerHeader.SoftPackageReferencesSerialInfo.Offset < 0)
 		{
