IoDispatcherFileBackend assertion fix

Check is commented out as it checks something that is no longer true and not required. Kept the original check for the reference with an added explanation for future reference.

#rnx

diff --git a/Engine/Source/Runtime/PakFile/Private/IoDispatcherFileBackend.cpp b/Engine/Source/Runtime/PakFile/Private/IoDispatcherFileBackend.cpp
--- a/Engine/Source/Runtime/PakFile/Private/IoDispatcherFileBackend.cpp
+++ b/Engine/Source/Runtime/PakFile/Private/IoDispatcherFileBackend.cpp
@@ -1983,14 +1983,14 @@ TIoStatusOr<FIoMappedRegion> FFileIoStore::OpenMapped(const FIoChunkId& ChunkId,
 			int32 BlockIndex = int32(ResolvedOffset / ContainerFile->CompressionBlockSize);
 			const FIoStoreTocCompressedBlockEntry& CompressionBlockEntry = ContainerFile->CompressionBlocks[BlockIndex];
 			const int64 BlockOffset = (int64)CompressionBlockEntry.GetOffset();
-			check(BlockOffset > 0 && IsAligned(BlockOffset, FPlatformProperties::GetMemoryMappingAlignment()));
-
+			// This is not required as mapped file region implementation ensure we map an aligned
+			// pointer/size, handling arbitrary offset/size.
+			// check(BlockOffset > 0 && IsAligned(BlockOffset, FPlatformProperties::GetMemoryMappingAlignment()));
 			IMappedFileHandle* MappedFileHandle = Reader->GetMappedContainerFileHandle(BlockOffset);
 			IMappedFileRegion* MappedFileRegion = MappedFileHandle->MapRegion(BlockOffset + Options.GetOffset(), ResolvedSize);
 			if (MappedFileRegion != nullptr)
 			{
-				// This is not required as mapped file region implementation ensure we map an aligned
-				// pointer/size. MappedPtr could be inside of the page, depending on the offset. This
+				// MappedPtr could be inside of the page, depending on the offset. This
 				// is actually fine, unless client code cares about MappedPtr alignment, which is
 				// public API.
 				// check(IsAligned(MappedFileRegion->GetMappedPtr(), FPlatformMemory::GetStats().PageSize));
