Replace FString with FUtf8String for a key value in FPakDirectory

diff --git a/Engine/Source/Runtime/PakFile/Private/IPlatformFilePak.cpp b/Engine/Source/Runtime/PakFile/Private/IPlatformFilePak.cpp
--- a/Engine/Source/Runtime/PakFile/Private/IPlatformFilePak.cpp
+++ b/Engine/Source/Runtime/PakFile/Private/IPlatformFilePak.cpp
@@ -5762,7 +5762,7 @@ void FPakPlatformFile::InitializeNewAsyncIO()
 uint64 GetRecursiveAllocatedSize(const FPakDirectory& Index)
 {
 	uint64 Size = Index.GetAllocatedSize();
-	for (const TPair<FString, FPakEntryLocation>& kvpair : Index)
+	for (const TPair<FUtf8String, FPakEntryLocation>& kvpair : Index)
 	{
 		Size += kvpair.Key.GetAllocatedSize();
 	}
diff --git a/Engine/Source/Runtime/PakFile/Private/PakFile.cpp b/Engine/Source/Runtime/PakFile/Private/PakFile.cpp
--- a/Engine/Source/Runtime/PakFile/Private/PakFile.cpp
+++ b/Engine/Source/Runtime/PakFile/Private/PakFile.cpp
@@ -699,8 +699,8 @@ bool FPakFile::LoadIndexInternal(FArchive& Reader, FDirectoryIndex& OutDirectory
 void FPakFile::LoadIndexInternal_DirectoryIndex(FArchive& Ar, FDirectoryIndex& OutDirectoryTMap,
 	FDirectoryTreeIndex& OutDirectoryTree, bool bLoadIntoDirectoryTMap, bool bLoadIntoDirectoryTree)
 {
-	FString DirectoryName;
 	FString FileName;
+	FString DirectoryName;
 	FPakEntryLocation FileData;
 	FPakDirectory DirectoryData;
 
@@ -722,11 +722,25 @@ void FPakFile::LoadIndexInternal_DirectoryIndex(FArchive& Ar, FDirectoryIndex& O
 
 		DirectoryData.Reset();
 		DirectoryData.Reserve(NumFiles);
-		for (; NumFiles > 0; --NumFiles)
+
+		if (Info.Version >= FPakInfo::PakFile_Version_Utf8PakDirectory)
 		{
-			Ar << FileName;
-			Ar << FileData;
-			DirectoryData.Add(FileName, FileData);
+			FUtf8String FileNameUtf8;
+			for (; NumFiles > 0; --NumFiles)
+			{
+				Ar << FileNameUtf8;
+				Ar << FileData;
+				DirectoryData.Add(MoveTemp(FileNameUtf8), FileData);
+			}
+		}
+		else
+		{
+			for (; NumFiles > 0; --NumFiles)
+			{
+				Ar << FileName;
+				Ar << FileData;
+				DirectoryData.Add(FUtf8String(FileName), FileData);
+			}
 		}
 
 #if ENABLE_PAKFILE_USE_DIRECTORY_TREE
@@ -757,9 +771,9 @@ void FPakFile::SaveIndexInternal_DirectoryIndex(FArchive& Ar, const FDirectoryIn
 		int32 NumFiles = DirPair.Value.Num();
 		Ar << NumFiles;
 
-		for (const TPair<FString, FPakEntryLocation>& FilePair : DirPair.Value)
+		for (const TPair<FUtf8String, FPakEntryLocation>& FilePair : DirPair.Value)
 		{
-			Ar << const_cast<FString&>(FilePair.Key);
+			Ar << const_cast<FUtf8String&>(FilePair.Key);
 			Ar << const_cast<FPakEntryLocation&>(FilePair.Value);
 		}
 	}
@@ -1005,10 +1019,10 @@ void FPakFile::PruneDirectoryIndexInternal(FDirectoryIndex* InOutDirectoryIndex,
 				FPakDirectory* PrunedDirectory = nullptr;
 				bool bKeepDirectory = false;
 
-				TArray<FString> RemoveFilenames;
+				TArray<FUtf8String> RemoveFilenames;
 				for (auto FileIt = It->Value.CreateIterator(); FileIt; ++FileIt)
 				{
-					const FString& FileNameWithoutPath = FileIt->Key;
+					const FUtf8String& FileNameWithoutPath = FileIt->Key;
 					const FString FullFilename = PakPathCombine(FullDirectoryName, FileNameWithoutPath);
 					bool bKeepFile = false;
 
@@ -1043,7 +1057,7 @@ void FPakFile::PruneDirectoryIndexInternal(FDirectoryIndex* InOutDirectoryIndex,
 				}
 				if (!InPrunedDirectoryIndex)
 				{
-					for (const FString& FileNameWithoutPath : RemoveFilenames)
+					for (const FUtf8String& FileNameWithoutPath : RemoveFilenames)
 					{
 						OriginalDirectory.Remove(FileNameWithoutPath);
 					}
@@ -1138,10 +1152,10 @@ void FPakFile::PruneDirectoryIndexInternal(FDirectoryIndex* InOutDirectoryIndex,
 				FPakDirectory* PrunedDirectory = nullptr;
 				bool bKeepDirectory = false;
 
-				TArray<FString> RemoveFilenames;
+				TArray<FUtf8String> RemoveFilenames;
 				for (FPakDirectory::TIterator FileIt = OriginalDirectory.CreateIterator(); FileIt; ++FileIt)
 				{
-					const FString& FileNameWithoutPath = FileIt->Key;
+					const FUtf8String& FileNameWithoutPath = FileIt->Key;
 					const FString FullFilename = PakPathCombine(FullDirectoryName, FileNameWithoutPath);
 					bool bKeepFile = false;
 
@@ -1176,7 +1190,7 @@ void FPakFile::PruneDirectoryIndexInternal(FDirectoryIndex* InOutDirectoryIndex,
 				}
 				if (!InPrunedDirectoryTreeIndex)
 				{
-					for (const FString& FileNameWithoutPath : RemoveFilenames)
+					for (const FUtf8String& FileNameWithoutPath : RemoveFilenames)
 					{
 						OriginalDirectory.Remove(FileNameWithoutPath);
 					}
@@ -1278,6 +1292,90 @@ void FPakFile::PruneDirectoryIndexInternal(FDirectoryIndex* InOutDirectoryIndex,
 	}
 }
 
+template <typename CharType>
+FString PakPathCombineInternal(FStringView Parent, TStringView<CharType> Child)
+{
+	// Our paths are different than FPaths, because our dirs / at the end, and "/" is the relative path to the mountdirectory and should be mapped to the empty string when joining
+	check(Parent.Len() > 0 && Parent[Parent.Len() - 1] == TEXT('/'));
+	if (Parent.Len() == 1)
+	{
+		return FString(Child);
+	}
+	else if (Child.Len() == 1 && Child[0] == CharType('/'))
+	{
+		return FString(Parent);
+	}
+	else
+	{
+		check(Child.Len() == 0 || Child[0] != CharType('/'));
+		FString Result;
+		Result.Reserve(Parent.Len() + Child.Len());
+		Result += Parent;
+		Result += Child;
+		return Result;
+	}
+}
+
+FString FPakFile::PakPathCombine(FStringView Parent, FStringView Child)
+{
+	return PakPathCombineInternal(Parent, Child);
+}
+
+FString FPakFile::PakPathCombine(FStringView Parent, FUtf8StringView Child)
+{
+	return PakPathCombineInternal(Parent, Child);
+}
+
+bool FPakFile::SplitPathInline(FStringView& InOutPath, FStringView& OutFilename)
+{
+	// FPaths::GetPath doesn't handle our / at the end of directories, so we have to do string manipulation ourselves
+	// The manipulation is less complicated than GetPath deals with, since we have normalized/path/strings, we have relative paths only, and we don't care about extensions
+	if (InOutPath.Len() == 0)
+	{
+		check(false); // Filenames should have non-zero length, and the minimum directory length is 1 (The root directory is written as "/")
+		return false;
+	}
+	else if (InOutPath.Len() == 1)
+	{
+		if (InOutPath[0] == TEXT('/'))
+		{
+			// The root directory; it has no parent.
+			OutFilename = FStringView();
+			return false;
+		}
+		else
+		{
+			// A relative one-character path with no /; this is a direct child of in the root directory
+			OutFilename = TEXT("/");
+			Swap(OutFilename, InOutPath);
+			return true;
+		}
+	}
+	else
+	{
+		if (InOutPath[InOutPath.Len() - 1] == TEXT('/'))
+		{
+			// The input was a Directory; remove the trailing / since we don't keep those on the CleanFilename
+			InOutPath.LeftChopInline(1);
+		}
+
+		int32 Offset = 0;
+		if (InOutPath.FindLastChar(TEXT('/'), Offset))
+		{
+			int32 FilenameStart = Offset + 1;
+			OutFilename = InOutPath.Mid(FilenameStart);
+			InOutPath.LeftInline(FilenameStart); // The Parent Directory keeps the / at the end
+		}
+		else
+		{
+			// A relative path with no /; this is a direct child of in the root directory
+			OutFilename = TEXT("/");
+			Swap(OutFilename, InOutPath);
+		}
+		return true;
+	}
+}
+
 FPakFile::EFindResult FPakFile::GetPakEntry(const FPakEntryLocation& PakEntryLocation, FPakEntry* OutEntry) const
 {
 	return GetPakEntry(PakEntryLocation, OutEntry, EncodedPakEntries, Files, Info);
@@ -1444,7 +1542,7 @@ void FPakFile::AddEntryToIndex(const FString& Filename, const FPakEntryLocation&
 			// Add the new directory; this has to come after the addition of the parent directories because we want to use the pointer afterwards and adding other directories would invalidate it
 			Directory = &DirectoryIndex->AddByHash(RelativeDirectoryFromMountHash, FString(RelativeDirectoryFromMount));
 		}
-		Directory->Add(FString(CleanFileName), EntryLocation);
+		Directory->Add(FUtf8String(CleanFileName), EntryLocation);
 	}
 #if ENABLE_PAKFILE_USE_DIRECTORY_TREE
 	if (DirectoryTreeIndex)
@@ -1470,7 +1568,7 @@ void FPakFile::AddEntryToIndex(const FString& Filename, const FPakEntryLocation&
 			// Add the new directory; this has to come after the addition of the parent directories because we want to use the pointer afterwards and adding other directories would invalidate it
 			Directory = &DirectoryTreeIndex->FindOrAdd(RelativeDirectoryFromMount);
 		}
-		Directory->Add(FString(CleanFileName), EntryLocation);
+		Directory->Add(FUtf8String(CleanFileName), EntryLocation);
 	}
 #endif
 
@@ -2215,7 +2313,7 @@ const FPakEntryLocation* FPakFile::FindLocationFromIndex(const FString& FullPath
 	const FPakDirectory* PakDirectory = DirectoryIndex.FindByHash(GetTypeHash(RelativeDirName), RelativeDirName);
 	if (PakDirectory)
 	{
-		return PakDirectory->FindByHash(GetTypeHash(CleanFileName), CleanFileName);
+		return PakDirectory->FindByHash(GetTypeHash(CleanFileName), FUtf8String(CleanFileName));
 	}
 	return nullptr;
 }
@@ -2238,7 +2336,7 @@ const FPakEntryLocation* FPakFile::FindLocationFromIndex(const FString& FullPath
 				const FPakDirectory* PakDirectory = InDirectoryTreeIndex.Find(RelativeDirName);
 				if (PakDirectory)
 				{
-					Result = PakDirectory->FindByHash(GetTypeHash(CleanFileName), CleanFileName);
+					Result = PakDirectory->FindByHash(GetTypeHash(CleanFileName), FUtf8String(CleanFileName));
 				}
 			}
 		}
diff --git a/Engine/Source/Runtime/PakFile/Private/PakFile.inl b/Engine/Source/Runtime/PakFile/Private/PakFile.inl
--- a/Engine/Source/Runtime/PakFile/Private/PakFile.inl
+++ b/Engine/Source/Runtime/PakFile/Private/PakFile.inl
@@ -174,13 +174,16 @@ void FPakFile::FindFilesAtPathInPakDirectoryInternal(const FString& MountPoint,
 		// Add everything
 		if (VisitFilter.bIncludeFiles)
 		{
+			TStringBuilder<1024> FilePathUnderDirectory;
 			for (FPakDirectory::TConstIterator FileIt(PakDirectory); FileIt; ++FileIt)
 			{
-				const FString& FilePathUnderDirectory = FileIt.Key();
-				if (VisitFilter.ShouldVisit(FilePathUnderDirectory))
+				FilePathUnderDirectory << FileIt.Key();
+				if (VisitFilter.ShouldVisit(FilePathUnderDirectory.ToView()))
 				{
-					OutFiles.Add(PakPathCombine(FullPathInIndex, FilePathUnderDirectory));
+					OutFiles.Add(PakPathCombine(FullPathInIndex, FilePathUnderDirectory.ToView()));
 				}
+
+				FilePathUnderDirectory.Reset();
 			}
 		}
 		if (VisitFilter.bIncludeDirectories)
@@ -202,13 +205,16 @@ void FPakFile::FindFilesAtPathInPakDirectoryInternal(const FString& MountPoint,
 		// Add files in the specified folder only.
 		if (VisitFilter.bIncludeFiles && SubDirIndex == INDEX_NONE)
 		{
+			TStringBuilder<1024> FilePathUnderDirectory;
 			for (FPakDirectory::TConstIterator FileIt(PakDirectory); FileIt; ++FileIt)
 			{
-				const FString& FilePathUnderDirectory = FileIt.Key();
-				if (VisitFilter.ShouldVisit(FilePathUnderDirectory))
+				FilePathUnderDirectory << FileIt.Key();
+				if (VisitFilter.ShouldVisit(FilePathUnderDirectory.ToView()))
 				{
-					OutFiles.Add(PakPathCombine(FullPathInIndex, FilePathUnderDirectory));
+					OutFiles.Add(PakPathCombine(FullPathInIndex, FilePathUnderDirectory.ToView()));
 				}
+
+				FilePathUnderDirectory.Reset();
 			}
 		}
 		// Add sub-folders in the specified folder only
diff --git a/Engine/Source/Runtime/PakFile/Public/IPlatformFilePak.h b/Engine/Source/Runtime/PakFile/Public/IPlatformFilePak.h
--- a/Engine/Source/Runtime/PakFile/Public/IPlatformFilePak.h
+++ b/Engine/Source/Runtime/PakFile/Public/IPlatformFilePak.h
@@ -10,6 +10,7 @@
 #include "Containers/SparseArray.h"
 #include "Containers/StringConv.h"
 #include "Containers/Ticker.h"
+#include "Containers/Utf8String.h"
 #include "Containers/UnrealString.h"
 #include "CoreMinimal.h"
 #include "Delegates/Delegate.h"
@@ -161,6 +162,7 @@ struct FPakInfo
 		PakFile_Version_FrozenIndex = 9,
 		PakFile_Version_PathHashIndex = 10,
 		PakFile_Version_Fnv64BugFix = 11,
+		PakFile_Version_Utf8PakDirectory = 12,
 
 
 		PakFile_Version_Last,
@@ -722,7 +724,7 @@ public:
 };
 
 /** Pak directory type mapping a filename to an FPakEntryLocation. */
-typedef TMap<FString, FPakEntryLocation> FPakDirectory;
+typedef TMap<FUtf8String, FPakEntryLocation> FPakDirectory;
 
 namespace UE::PakFile::Private
 {
@@ -1411,79 +1413,11 @@ public:
 	}
 
 	/* Helper function to join two path strings that are in the PakPath format */
-	static FString PakPathCombine(FStringView Parent, FStringView Child)
-	{
-		// Our paths are different than FPaths, because our dirs / at the end, and "/" is the relative path to the mountdirectory and should be mapped to the empty string when joining
-		check(Parent.Len() > 0 && Parent[Parent.Len() - 1] == TEXT('/'));
-		if (Parent.Len() == 1)
-		{
-			return FString(Child);
-		}
-		else if (Child.Len() == 1 && Child[0] == TEXT('/'))
-		{
-			return FString(Parent);
-		}
-		else
-		{
-			check(Child.Len() == 0 || Child[0] != TEXT('/'));
-			FString Result;
-			Result.Reserve(Parent.Len() + Child.Len());
-			Result += Parent;
-			Result += Child;
-			return Result;
-		}
-	}
+	static FString PakPathCombine(FStringView Parent, FStringView Child);
+	static FString PakPathCombine(FStringView Parent, FUtf8StringView Child);
 
 	/** Helper function to split a PakDirectoryIndex-Formatted PathName into its PakDirectoryIndex-Formatted parent directory and the CleanFileName */
-	static bool SplitPathInline(FStringView& InOutPath, FStringView& OutFilename)
-	{
-		// FPaths::GetPath doesn't handle our / at the end of directories, so we have to do string manipulation ourselves
-		// The manipulation is less complicated than GetPath deals with, since we have normalized/path/strings, we have relative paths only, and we don't care about extensions
-		if (InOutPath.Len() == 0)
-		{
-			check(false); // Filenames should have non-zero length, and the minimum directory length is 1 (The root directory is written as "/")
-			return false;
-		}
-		else if (InOutPath.Len() == 1)
-		{
-			if (InOutPath[0] == TEXT('/'))
-			{
-				// The root directory; it has no parent.
-				OutFilename = FStringView();
-				return false;
-			}
-			else
-			{
-				// A relative one-character path with no /; this is a direct child of in the root directory
-				OutFilename = TEXT("/");
-				Swap(OutFilename, InOutPath);
-				return true;
-			}
-		}
-		else
-		{
-			if (InOutPath[InOutPath.Len() - 1] == TEXT('/'))
-			{
-				// The input was a Directory; remove the trailing / since we don't keep those on the CleanFilename
-				InOutPath.LeftChopInline(1);
-			}
-
-			int32 Offset = 0;
-			if (InOutPath.FindLastChar(TEXT('/'), Offset))
-			{
-				int32 FilenameStart = Offset + 1;
-				OutFilename = InOutPath.Mid(FilenameStart);
-				InOutPath.LeftInline(FilenameStart); // The Parent Directory keeps the / at the end
-			}
-			else
-			{
-				// A relative path with no /; this is a direct child of in the root directory
-				OutFilename = TEXT("/");
-				Swap(OutFilename, InOutPath);
-			}
-			return true;
-		}
-	}
+	static bool SplitPathInline(FStringView& InOutPath, FStringView& OutFilename);
 
 	static bool SplitPathInline(FString& InOutPath, FString& OutFilename)
 	{
@@ -2522,8 +2456,8 @@ public:
 					FPakEntry PakEntry;
 					if (PakFile->GetPakEntry(FileIt.Value(), &PakEntry) != FPakFile::EFindResult::NotFound && PakEntry.Offset == FileEntry.Offset)
 					{
-						const FString& RealFilename = FileIt.Key();
-						return Path / RealFilename;
+						const FUtf8String& RealFilename = FileIt.Key();
+						return Path / FString(RealFilename);
 					}
 				}
 			}
