File iodispatcher -  Fix inverted error condition leading to errors about trying to read soft refs when soft refs had been loaded.

diff --git a/Engine/Source/Runtime/PakFile/Private/IoDispatcherFileBackend.cpp b/Engine/Source/Runtime/PakFile/Private/IoDispatcherFileBackend.cpp
--- a/Engine/Source/Runtime/PakFile/Private/IoDispatcherFileBackend.cpp
+++ b/Engine/Source/Runtime/PakFile/Private/IoDispatcherFileBackend.cpp
@@ -934,7 +934,6 @@ TIoStatusOr<FIoContainerHeader> FFileIoStoreReader::ReadContainerHeader(bool bRe
 		ContainerHeader = FIoContainerHeader();
 	}
 
-	// todo, ensure here if the container didn't include softreferences and we want to load them, the size check isn't really sufficent
 	if (bReadSoftRefs)
 	{
 		if (ContainerHeader.SoftPackageReferencesSerialInfo.Offset < 0)
@@ -1325,8 +1324,8 @@ TIoStatusOr<FIoContainerHeader> FFileIoStore::Mount(
 		}
 	}
 
-	TIoStatusOr<FIoContainerHeader> ContainerHeaderReadResult = Reader->ReadContainerHeader(
-		EnumHasAnyFlags(Options, UE::IoStore::ETocMountOptions::WithSoftReferences));
+	const bool bWithSoftRefs = EnumHasAnyFlags(Options, UE::IoStore::ETocMountOptions::WithSoftReferences);
+	TIoStatusOr<FIoContainerHeader> ContainerHeaderReadResult = Reader->ReadContainerHeader(bWithSoftRefs);
 	FIoContainerHeader ContainerHeader;
 	if (ContainerHeaderReadResult.IsOk())
 	{
@@ -1349,7 +1348,12 @@ TIoStatusOr<FIoContainerHeader> FFileIoStore::Mount(
 			return A->GetContainerInstanceId() > B->GetContainerInstanceId();
 		});
 		IoStoreReaders.Insert(MoveTemp(Reader), InsertionIndex);
-		UE_LOG(LogIoDispatcher, Display, TEXT("Mounting container '%s' in location slot %d"), InTocPath, InsertionIndex);
+
+		UE_LOG(LogIoDispatcher, Display, TEXT("Mounting container '%s' in location slot %d (mount options '%s'%s)"),
+			InTocPath, InsertionIndex, LexToString(Options),
+			bWithSoftRefs
+				? *FString::Printf(TEXT(", soft reference count: %d"), ContainerHeader.SoftPackageReferences.PackageIds.Num())
+				: TEXT(""));
 	}
 
 	return ContainerHeader;
