File iodispatcher -  Fix inverted error condition leading to errors about trying to read soft refs when soft refs had been loaded.
- trying this again with the correct files

diff --git a/Engine/Source/Runtime/PakFile/Private/FilePackageStore.cpp b/Engine/Source/Runtime/PakFile/Private/FilePackageStore.cpp
--- a/Engine/Source/Runtime/PakFile/Private/FilePackageStore.cpp
+++ b/Engine/Source/Runtime/PakFile/Private/FilePackageStore.cpp
@@ -401,7 +401,7 @@ TConstArrayView<uint32> FFilePackageStoreBackend::GetSoftReferences(FPackageId P
 			if (Container != nullptr)
 			{
 				// to do: figure out a better way to make sure softreference information wasn't loaded for this container, also this should make the outer request fail completely and error with appropriate information
-				if (Container->ContainerHeader->SoftPackageReferences.bLoadedSoftPackageReferences)
+				if (!Container->ContainerHeader->SoftPackageReferences.bLoadedSoftPackageReferences)
 				{
 					UE_LOG(LogFilePackageStore, Error, TEXT("Attempting to read softreference information for package %s but container %s never loaded softreference information"), *LexToString(PackageId), *LexToString(Container->ContainerHeader->ContainerId));
 				}
