Disable soft ref load warning for FFilePackageStoreBackend.  This currently spams the log with invalid warnings if the package entry has been deduplicated.

diff --git a/Engine/Source/Runtime/PakFile/Private/FilePackageStore.cpp b/Engine/Source/Runtime/PakFile/Private/FilePackageStore.cpp
--- a/Engine/Source/Runtime/PakFile/Private/FilePackageStore.cpp
+++ b/Engine/Source/Runtime/PakFile/Private/FilePackageStore.cpp
@@ -416,20 +416,20 @@ TConstArrayView<uint32> FFilePackageStoreBackend::GetSoftReferences(FPackageId P
 		else
 		{
 			// make sure tha tit's containers mounted soft refs
-#if !UE_BUILD_SHIPPING
-			FMountedContainer* Container = Algo::FindByPredicate(
-				MountedContainers,
-				[Entry](const FMountedContainer& C) { return C.EntryDataRange.Contains(Entry->Offset); });
-
-			if (Container != nullptr)
-			{
-				// to do: figure out a better way to make sure softreference information wasn't loaded for this container, also this should make the outer request fail completely and error with appropriate information
-				if (!Container->ContainerHeader->SoftPackageReferences.bLoadedSoftPackageReferences)
-				{
-					UE_LOG(LogFilePackageStore, Error, TEXT("Attempting to read softreference information for package %s but container %s never loaded softreference information"), *LexToString(PackageId), *LexToString(Container->ContainerHeader->ContainerId));
-				}
-			}
-#endif
+//#if !UE_BUILD_SHIPPING
+//			FMountedContainer* Container = Algo::FindByPredicate(
+//				MountedContainers,
+//				[Entry](const FMountedContainer& C) { return C.EntryDataRange.Contains(Entry->Offset); });
+//
+//			if (Container != nullptr)
+//			{
+//				// to do: figure out a better way to make sure softreference information wasn't loaded for this container, also this should make the outer request fail completely and error with appropriate information
+//				if (!Container->ContainerHeader->SoftPackageReferences.bLoadedSoftPackageReferences)
+//				{
+//					UE_LOG(LogFilePackageStore, Error, TEXT("Attempting to read softreference information for package %s but container %s never loaded softreference information"), *LexToString(PackageId), *LexToString(Container->ContainerHeader->ContainerId));
+//				}
+//			}
+//#endif
 		}
 	}
 
