Updated mount logging with info about container flags, mount options and whether the container has soft references or not.

#rnx

diff --git a/Engine/Source/Runtime/Core/Private/IO/IoDispatcher.cpp b/Engine/Source/Runtime/Core/Private/IO/IoDispatcher.cpp
--- a/Engine/Source/Runtime/Core/Private/IO/IoDispatcher.cpp
+++ b/Engine/Source/Runtime/Core/Private/IO/IoDispatcher.cpp
@@ -72,6 +72,43 @@ static_assert(std::size(GetIoErrorText_ErrorCodeTextArray) == static_cast<uint32
 
 const TCHAR* const* GetIoErrorText_ErrorCodeText = GetIoErrorText_ErrorCodeTextArray;
 
+FStringBuilderBase& operator<<(FStringBuilderBase& Sb, EIoContainerFlags Flags)
+{
+	static const TCHAR* Names[] =
+	{
+		TEXT("None"),
+		TEXT("Compressed"),
+		TEXT("Encrypted"),
+		TEXT("Signed"),
+		TEXT("Indexed"),
+		TEXT("OnDemand")
+	};
+	constexpr uint32 BitCount = 1 + FMath::CountTrailingZeros(
+		static_cast<std::underlying_type_t<EIoContainerFlags>>(EIoContainerFlags::Last));
+	static_assert(UE_ARRAY_COUNT(Names) == BitCount + 1, "Please update names list");
+
+	for (int32 Idx = 0; Idx < BitCount; ++Idx)
+	{
+		const EIoContainerFlags FlagToTest = static_cast<EIoContainerFlags>(1 << Idx);
+		if (EnumHasAnyFlags(Flags, FlagToTest))
+		{
+			if (Sb.Len())
+			{
+				Sb << TEXT("|");
+			}
+			Sb << Names[Idx + 1];
+		}
+	}
+
+	return Sb;
+}
+
+FString LexToString(EIoContainerFlags Flags)
+{
+	TStringBuilder<128> Sb;
+	return (Sb << Flags).ToString();
+}
+
 #if UE_IODISPATCHER_STATS_ENABLED
 class FIoRequestStats
 {
diff --git a/Engine/Source/Runtime/Core/Public/IO/IoDispatcher.h b/Engine/Source/Runtime/Core/Public/IO/IoDispatcher.h
--- a/Engine/Source/Runtime/Core/Public/IO/IoDispatcher.h
+++ b/Engine/Source/Runtime/Core/Public/IO/IoDispatcher.h
@@ -484,9 +484,13 @@ enum class EIoContainerFlags : uint8
 	Signed		= (1 << 2),
 	Indexed		= (1 << 3),
 	OnDemand	= (1 << 4),
+	Last		= OnDemand
 };
 ENUM_CLASS_FLAGS(EIoContainerFlags);
 
+CORE_API FStringBuilderBase& operator<<(FStringBuilderBase& Sb, EIoContainerFlags Flags);
+CORE_API FString LexToString(EIoContainerFlags Flags);
+
 struct FIoContainerSettings
 {
 	FIoContainerId ContainerId;
diff --git a/Engine/Source/Runtime/PakFile/Private/FileIoDispatcherBackend.h b/Engine/Source/Runtime/PakFile/Private/FileIoDispatcherBackend.h
--- a/Engine/Source/Runtime/PakFile/Private/FileIoDispatcherBackend.h
+++ b/Engine/Source/Runtime/PakFile/Private/FileIoDispatcherBackend.h
@@ -14,9 +14,10 @@ namespace UE::IoStore
 
 enum class ETocMountOptions
 {
-	None = 0,
+	None				= 0,
 	/** Make soft references available */
-	WithSoftReferences = 1 << 0,
+	WithSoftReferences	= 1 << 0,
+	Last				= WithSoftReferences
 };
 ENUM_CLASS_FLAGS(ETocMountOptions);
 
@@ -39,3 +40,16 @@ public:
 TSharedPtr<IFileIoDispatcherBackend> MakeFileIoDispatcherBackend();
 
 } // namespace UE::IoStore
+
+inline const TCHAR* LexToString(UE::IoStore::ETocMountOptions MountOptions)
+{
+	static const TCHAR* Names[]
+	{
+		TEXT("None"),
+		TEXT("WithSoftReferences")
+	};
+	static_assert(std::size(Names) == (1 + SIZE_T(UE::IoStore::ETocMountOptions::Last)));
+
+	// Update this if more options are added
+	return Names[int32(MountOptions)];
+}
diff --git a/Engine/Source/Runtime/PakFile/Private/IoDispatcherFileBackend.cpp b/Engine/Source/Runtime/PakFile/Private/IoDispatcherFileBackend.cpp
--- a/Engine/Source/Runtime/PakFile/Private/IoDispatcherFileBackend.cpp
+++ b/Engine/Source/Runtime/PakFile/Private/IoDispatcherFileBackend.cpp
@@ -659,7 +659,7 @@ FIoStatus FFileIoStoreReader::Initialize(const TCHAR* InTocFilePath, int32 InOrd
 
 	IPlatformFile& Ipf = FPlatformFileManager::Get().GetPlatformFile();
 
-	UE_LOG(LogIoDispatcher, Display, TEXT("Reading toc: %s"), InTocFilePath);
+	UE_LOG(LogIoDispatcher, Verbose, TEXT("Reading toc: %s"), InTocFilePath);
 
 	FIoStoreTocResourceView TocResource;
 	FIoStoreTocResourceStorage TocStorage;
@@ -721,7 +721,7 @@ FIoStatus FFileIoStoreReader::Initialize(const TCHAR* InTocFilePath, int32 InOrd
 
 	Stats.OnTocMounted(GetTocAllocatedSize());
 
-	UE_LOG(LogIoDispatcher, Display, TEXT("Toc loaded : %s, Id=%s, Order=%d, EntryCount=%u, SignatureHash=%s"),
+	UE_LOG(LogIoDispatcher, Verbose, TEXT("Toc loaded : %s, Id=%s, Order=%d, EntryCount=%u, SignatureHash=%s"),
 		InTocFilePath, *LexToString(TocResource.Header.ContainerId), InOrder, TocResource.Header.TocEntryCount, *TocResource.SignatureHash.ToString());
 
 	// Hold onto the handle and mapped regions
@@ -1348,8 +1348,11 @@ TIoStatusOr<FIoContainerHeader> FFileIoStore::Mount(
 			}
 			return A->GetContainerInstanceId() > B->GetContainerInstanceId();
 		});
+		const EIoContainerFlags ContainerFlags = Reader->GetContainerFile() != nullptr ? Reader->GetContainerFile()->ContainerFlags : EIoContainerFlags::None;
 		IoStoreReaders.Insert(MoveTemp(Reader), InsertionIndex);
-		UE_LOG(LogIoDispatcher, Display, TEXT("Mounting container '%s' in location slot %d"), InTocPath, InsertionIndex);
+		UE_LOG(LogIoDispatcher, Display, TEXT("Mounted container '%s', Id='%s', Order=%d, Slot=%d, Options=(%s), Flags=(%s), HasSoftRefs=%s"),
+			InTocPath, *LexToString(ContainerHeader.ContainerId), Order, InsertionIndex, LexToString(Options), *LexToString(ContainerFlags),
+			ContainerHeader.SoftPackageReferences.bContainsSoftPackageReferences ? TEXT("True") : TEXT("False"));
 	}
 
 	return ContainerHeader;
