Return empty list of soft reference indices for file package store entries with no dependecies.
- this change fixes an issue where a warning message prints out the wrong container ID since
all containers has an entry for empty package store entries.

#rnx

diff --git a/Engine/Source/Runtime/PakFile/Private/FilePackageStore.cpp b/Engine/Source/Runtime/PakFile/Private/FilePackageStore.cpp
--- a/Engine/Source/Runtime/PakFile/Private/FilePackageStore.cpp
+++ b/Engine/Source/Runtime/PakFile/Private/FilePackageStore.cpp
@@ -392,6 +392,12 @@ TConstArrayView<uint32> FFilePackageStoreBackend::GetSoftReferences(FPackageId P
 {
 	if (const FEntryHandle* Entry = PackageEntries.Find(PackageId); Entry != nullptr)
 	{
+		if (IsEmpty(*Entry))
+		{
+			OutPackageIds = TConstArrayView<FPackageId>();
+			return TConstArrayView<uint32>();
+		}
+
 		if (Entry->HasSoftRefs)
 		{
 			FMountedContainer* Container = Algo::FindByPredicate(
