Do not look for .metalmap outside of paks

Metalmaps are now always in UFS as code to stage/load those in the usual file system was removed in 38074494.
This saves 2 syscalls per RHICreateShaderLibrary call, as it was falling back to FS lookup.

[RN]

diff --git a/Engine/Source/Runtime/PakFile/Private/IPlatformFilePak.cpp b/Engine/Source/Runtime/PakFile/Private/IPlatformFilePak.cpp
--- a/Engine/Source/Runtime/PakFile/Private/IPlatformFilePak.cpp
+++ b/Engine/Source/Runtime/PakFile/Private/IPlatformFilePak.cpp
@@ -5626,6 +5626,7 @@ bool FPakPlatformFile::Initialize(IPlatformFile* Inner, const TCHAR* CmdLine)
 	ExcludedNonPakExtensions.Add(TEXT("uexp"));
 	ExcludedNonPakExtensions.Add(TEXT("uptnl"));
 	ExcludedNonPakExtensions.Add(TEXT("ushaderbytecode"));
+	ExcludedNonPakExtensions.Add(TEXT("metalmap"));
 #endif
 
 #if DISABLE_NONUFS_INI_WHEN_COOKED
