[PAK] Prevent crash in RegisterEncryptionKey due to corrupted Pak file.

#rnx

diff --git a/Engine/Source/Runtime/PakFile/Private/IPlatformFilePak.cpp b/Engine/Source/Runtime/PakFile/Private/IPlatformFilePak.cpp
--- a/Engine/Source/Runtime/PakFile/Private/IPlatformFilePak.cpp
+++ b/Engine/Source/Runtime/PakFile/Private/IPlatformFilePak.cpp
@@ -6538,25 +6538,28 @@ void FPakPlatformFile::RegisterEncryptionKey(const FGuid& InGuid, const FAES::FA
 	int32 NumMounted = 0;
 	TSet<int32> ChunksToNotify;
 
-	for (const FPakListDeferredEntry& Entry : PendingEncryptedPakFiles)
 	{
-		if (Entry.EncryptionKeyGuid == InGuid)
+		const TArray<FPakListDeferredEntry> LocalPendingEncryptedPakFiles = PendingEncryptedPakFiles;
+		for (const FPakListDeferredEntry& Entry : LocalPendingEncryptedPakFiles)
 		{
-			if (Mount(*Entry.Filename, Entry.ReadOrder, Entry.Path.Len() == 0 ? nullptr : *Entry.Path))
+			if (Entry.EncryptionKeyGuid == InGuid)
 			{
-				UE_LOG(LogPakFile, Log, TEXT("Successfully mounted deferred pak file '%s'"), *Entry.Filename);
-				NumMounted++;
+				if (Mount(*Entry.Filename, Entry.ReadOrder, Entry.Path.Len() == 0 ? nullptr : *Entry.Path))
+				{
+					UE_LOG(LogPakFile, Log, TEXT("Successfully mounted deferred pak file '%s'"), *Entry.Filename);
+					NumMounted++;
 
-				int32 PakchunkIndex = GetPakchunkIndexFromPakFile(Entry.Filename);
-				if (PakchunkIndex != INDEX_NONE)
+					int32 PakchunkIndex = GetPakchunkIndexFromPakFile(Entry.Filename);
+					if (PakchunkIndex != INDEX_NONE)
+					{
+						ChunksToNotify.Add(PakchunkIndex);
+					}
+				}
+				else
 				{
-					ChunksToNotify.Add(PakchunkIndex);
+					UE_LOG(LogPakFile, Warning, TEXT("Failed to mount deferred pak file '%s'"), *Entry.Filename);
 				}
 			}
-			else
-			{
-				UE_LOG(LogPakFile, Warning, TEXT("Failed to mount deferred pak file '%s'"), *Entry.Filename);
-			}
 		}
 	}
 
