Fix for memory mapping iostore partitions

#rnx

diff --git a/Engine/Source/Runtime/PakFile/Private/IoDispatcherFileBackend.cpp b/Engine/Source/Runtime/PakFile/Private/IoDispatcherFileBackend.cpp
--- a/Engine/Source/Runtime/PakFile/Private/IoDispatcherFileBackend.cpp
+++ b/Engine/Source/Runtime/PakFile/Private/IoDispatcherFileBackend.cpp
@@ -41,8 +41,9 @@ TAtomic<uint32> FFileIoStoreReader::GlobalContainerInstanceId{ 0 };
 class FMappedFileProxy final : public IMappedFileHandle
 {
 public:
-	FMappedFileProxy(IMappedFileHandle* InSharedMappedFileHandle, uint64 InSize)
+	FMappedFileProxy(IMappedFileHandle* InSharedMappedFileHandle, uint64 InSize, uint64 InPartitionSize)
 		: IMappedFileHandle(InSize)
+		, PartitionSize(InPartitionSize)
 		, SharedMappedFileHandle(InSharedMappedFileHandle)
 	{
 	}
@@ -51,9 +52,15 @@ public:
 
 	virtual IMappedFileRegion* MapRegion(int64 Offset = 0, int64 BytesToMap = MAX_int64, FFileMappingFlags Flags = EMappedFileFlags::ENone) override
 	{
+		// The offset we get is the TocOffset for the container, and we need the offset into our partition.
+		if (PartitionSize)
+		{
+			Offset %= PartitionSize;
+		}
 		return SharedMappedFileHandle != nullptr ? SharedMappedFileHandle->MapRegion(Offset, BytesToMap, Flags) : nullptr;
 	}
 private:
+	uint64 PartitionSize;
 	IMappedFileHandle* SharedMappedFileHandle;
 };
 
@@ -844,7 +851,7 @@ IMappedFileHandle* FFileIoStoreReader::GetMappedContainerFileHandle(uint64 TocOf
 	}
 
 	check(Partition.FileSize > 0);
-	return new FMappedFileProxy(Partition.MappedFileHandle.Get(), Partition.FileSize);
+	return new FMappedFileProxy(Partition.MappedFileHandle.Get(), Partition.FileSize, ContainerFile.PartitionSize);
 }
 
 TIoStatusOr<FIoContainerHeader> FFileIoStoreReader::ReadContainerHeader(bool bReadSoftRefs) const
