Fixed broken printf specifier strings.

diff --git a/Engine/Plugins/Runtime/ResonanceAudio/Source/ResonanceAudio/Private/ResonanceAudioReverb.cpp b/Engine/Plugins/Runtime/ResonanceAudio/Source/ResonanceAudio/Private/ResonanceAudioReverb.cpp
--- a/Engine/Plugins/Runtime/ResonanceAudio/Source/ResonanceAudio/Private/ResonanceAudioReverb.cpp
+++ b/Engine/Plugins/Runtime/ResonanceAudio/Source/ResonanceAudio/Private/ResonanceAudioReverb.cpp
@@ -209,7 +209,7 @@ namespace ResonanceAudio
 			if (!IgnoreUserResonanceSubmixCVar)
 			{
 				static const FString DefaultSubmixName = TEXT("Resonance Reverb Submix");
-				UE_LOG(LogResonanceAudio, Error, TEXT("Couldn't load Resonance Reverb Submix from object path '%s' in ResonanceSettings. Creating stub submix."),
+				UE_LOG(LogResonanceAudio, Error, TEXT("Couldn't load Resonance Reverb Submix from object path '%s' in ResonanceSettings. Creating stub submix '%s'."),
 					*Settings->OutputSubmix.GetAssetPathString(),
 					*DefaultSubmixName);
 			}
diff --git a/Engine/Plugins/Runtime/USDCore/Source/USDUtilities/Private/USDMaterialXShaderGraph.cpp b/Engine/Plugins/Runtime/USDCore/Source/USDUtilities/Private/USDMaterialXShaderGraph.cpp
--- a/Engine/Plugins/Runtime/USDCore/Source/USDUtilities/Private/USDMaterialXShaderGraph.cpp
+++ b/Engine/Plugins/Runtime/USDCore/Source/USDUtilities/Private/USDMaterialXShaderGraph.cpp
@@ -602,7 +602,7 @@ mx::NodeDefPtr FUsdMaterialXShaderGraph::GetNodeDef(const pxr::UsdShadeConnectio
 			else if (pos = NodeDefString.rfind("string"); pos != std::string::npos)
 			{
 				USD_LOG_WARNING(
-					TEXT("'$s': '%s' are not supported"),
+					TEXT("'%s': '%s' are not supported"),
 					*UsdToUnreal::ConvertString(Source.source.GetPrim().GetName()),
 					*UsdToUnreal::ConvertString(NodeDefString)
 				);
diff --git a/Engine/Plugins/Runtime/USDCore/Source/USDUtilities/Private/USDSkeletalDataConversion.cpp b/Engine/Plugins/Runtime/USDCore/Source/USDUtilities/Private/USDSkeletalDataConversion.cpp
--- a/Engine/Plugins/Runtime/USDCore/Source/USDUtilities/Private/USDSkeletalDataConversion.cpp
+++ b/Engine/Plugins/Runtime/USDCore/Source/USDUtilities/Private/USDSkeletalDataConversion.cpp
@@ -4315,7 +4315,7 @@ bool UnrealToUsd::ConvertControlRigSection(
 								bRenamedAChannel = true;
 								BlendShapeChannels[BlendShapeIndex] = UnrealToUsd::ConvertToken(*FoundBlendShape->Name).Get();
 								USD_LOG_INFO(
-									TEXT("Updating Mesh '{0}' to bind BlendShape target '{1}' to SkelAnimation curve '{2}'"),
+									TEXT("Updating Mesh '%s' to bind BlendShape target '%s' to SkelAnimation curve '%s'"),
 									*MeshPrim.GetPrimPath().GetString(),
 									*PrimaryBlendShapePath,
 									*FoundBlendShape->Name
diff --git a/Engine/Source/Developer/DesktopPlatform/Private/TargetReceiptBuildWorker.cpp b/Engine/Source/Developer/DesktopPlatform/Private/TargetReceiptBuildWorker.cpp
--- a/Engine/Source/Developer/DesktopPlatform/Private/TargetReceiptBuildWorker.cpp
+++ b/Engine/Source/Developer/DesktopPlatform/Private/TargetReceiptBuildWorker.cpp
@@ -177,7 +177,7 @@ bool FTargetReceiptBuildWorker::TryAddExecutablePath(FStringView Path)
 {
 	if (Path.IsEmpty())
 	{
-		UE_LOG(LogTargetReceiptBuildWorker, Warning, TEXT("Empty executable path encountered while populating worker '%s'."), *Name, Path.Len(), Path.GetData());
+		UE_LOG(LogTargetReceiptBuildWorker, Warning, TEXT("Empty executable path encountered while populating worker '%s': '%.*s'"), *Name, Path.Len(), Path.GetData());
 		return false;
 	}
 
diff --git a/Engine/Source/Developer/Zen/Private/ZenServerInterface.cpp b/Engine/Source/Developer/Zen/Private/ZenServerInterface.cpp
--- a/Engine/Source/Developer/Zen/Private/ZenServerInterface.cpp
+++ b/Engine/Source/Developer/Zen/Private/ZenServerInterface.cpp
@@ -1084,7 +1084,7 @@ static FServiceAutoLaunchSettings::EInstallMode ZenGetInstallModeFromString(cons
 	}
 	else
 	{
-		checkf(false, TEXT("Invalid zenserver install mode: {InstallMode}"), *InstallMode);
+		checkf(false, TEXT("Invalid zenserver install mode: {%s}"), *InstallMode);
 		return FServiceAutoLaunchSettings::EInstallMode::Copy;
 	}
 }
diff --git a/Engine/Source/Editor/BlueprintGraph/Private/BlueprintActionDatabase.cpp b/Engine/Source/Editor/BlueprintGraph/Private/BlueprintActionDatabase.cpp
--- a/Engine/Source/Editor/BlueprintGraph/Private/BlueprintActionDatabase.cpp
+++ b/Engine/Source/Editor/BlueprintGraph/Private/BlueprintActionDatabase.cpp
@@ -2041,7 +2041,7 @@ bool FBlueprintActionDatabase::IsFunctionAllowed(UFunction const* InFunction, EP
 				{
 					if (!IsStructAllowed(StructProperty->Struct, EPermissionsContext::Pin))
 					{
-						UE_LOG(LogBlueprint, Warning, TEXT("Function %s was filtered because one of its parameters (struct %s) was filtered"), *InFunction->GetPathName(), *Property->GetName(), *StructProperty->Struct->GetPathName());
+						UE_LOG(LogBlueprint, Warning, TEXT("Function %s was filtered because one of its parameters %s (struct %s) was filtered"), *InFunction->GetPathName(), *Property->GetName(), *StructProperty->Struct->GetPathName());
 						return false;
 					}
 				}
diff --git a/Engine/Source/Editor/UnrealEd/Private/Factories/PackFactory.cpp b/Engine/Source/Editor/UnrealEd/Private/Factories/PackFactory.cpp
--- a/Engine/Source/Editor/UnrealEd/Private/Factories/PackFactory.cpp
+++ b/Engine/Source/Editor/UnrealEd/Private/Factories/PackFactory.cpp
@@ -718,7 +718,7 @@ UObject* UPackFactory::FactoryCreateBinary
 					}
 					else
 					{
-						UE_LOG(LogPackFactory, Error, TEXT("Unable to create directory \"%s\"."), *FileToCopy, *DestFilename);
+						UE_LOG(LogPackFactory, Error, TEXT("Unable to create directory \"%s\"."), *DestDirectory);
 					}
 				}
 			}
diff --git a/Engine/Source/Editor/UnrealEd/Private/UnrealEdMisc.cpp b/Engine/Source/Editor/UnrealEd/Private/UnrealEdMisc.cpp
--- a/Engine/Source/Editor/UnrealEd/Private/UnrealEdMisc.cpp
+++ b/Engine/Source/Editor/UnrealEd/Private/UnrealEdMisc.cpp
@@ -2344,7 +2344,7 @@ void FUnrealEdMisc::MountTemplateSharedPaths()
 				// Get the name of the level we are falling back to so we can tell the user
 				FString FallbackDetailString;
 				UEnum::GetValueAsString(TEXT("/Script/AddContentDialog.EFeaturePackDetailLevel"), EachEditDetail, FallbackDetailString);
-				UE_LOG(LogUnrealEdMisc, Verbose, TEXT("Only 1 detail level defined for %s in %s. Cannot edit detail level %s. Will fallback to  "), *EachPack.MountName, *TemplateFilename, *DetailString,*FallbackDetailString);
+				UE_LOG(LogUnrealEdMisc, Verbose, TEXT("Only 1 detail level defined for %s in %s. Cannot edit detail level %s. Will fallback to %s"), *EachPack.MountName, *TemplateFilename, *DetailString,*FallbackDetailString);
 				// Then correct the string too !
 				DetailString = FallbackDetailString;
 
diff --git a/Engine/Source/Runtime/Core/Private/Serialization/CompactBinarySerialization.cpp b/Engine/Source/Runtime/Core/Private/Serialization/CompactBinarySerialization.cpp
--- a/Engine/Source/Runtime/Core/Private/Serialization/CompactBinarySerialization.cpp
+++ b/Engine/Source/Runtime/Core/Private/Serialization/CompactBinarySerialization.cpp
@@ -649,7 +649,7 @@ private:
 	void AppendBase64String(FMemoryView Value)
 	{
 		checkf(Value.GetSize() <= 512 * 1024 * 1024,
-			TEXT("Encoding 512 MiB or larger is not supported. Size: " UINT64_FMT), Value.GetSize());
+			TEXT("Encoding 512 MiB or larger is not supported. Size: %" UINT64_FMT), Value.GetSize());
 		const uint32 EncodedSize = FBase64::GetEncodedDataSize(uint32(Value.GetSize()));
 		// Encode writes a zero-terminator after its text, so we need to have allocated space for it to write to by
 		// reserving at least one more than the EncodedSize. We handle that by reserving space for the trailing quote
diff --git a/Engine/Source/Runtime/CoreUObject/Private/UObject/PropertyOptional.cpp b/Engine/Source/Runtime/CoreUObject/Private/UObject/PropertyOptional.cpp
--- a/Engine/Source/Runtime/CoreUObject/Private/UObject/PropertyOptional.cpp
+++ b/Engine/Source/Runtime/CoreUObject/Private/UObject/PropertyOptional.cpp
@@ -267,7 +267,7 @@ void FOptionalProperty::SerializeItem(FStructuredArchive::FSlot Slot, void* Data
 				const int64 NumLoadedBytes = UnderlyingArchive.Tell() - ValueStartOffset;
 				checkf(
 					ValueTag.Size == NumLoadedBytes,
-					TEXT("Value tagged with ID %s was saved as %i bytes, but loading it as a %s loaded " INT64_FMT " bytes."),
+					TEXT("Value tagged with ID %s was saved as %i bytes, but loading it as a %s loaded %" INT64_FMT " bytes."),
 					*ValueTag.Type.ToString(),
 					ValueTag.Size,
 					*GetValueProperty()->GetClass()->GetFName().ToString(),
diff --git a/Engine/Source/Runtime/CoreUObject/Private/UObject/RemoteObjectTransfer.cpp b/Engine/Source/Runtime/CoreUObject/Private/UObject/RemoteObjectTransfer.cpp
--- a/Engine/Source/Runtime/CoreUObject/Private/UObject/RemoteObjectTransfer.cpp
+++ b/Engine/Source/Runtime/CoreUObject/Private/UObject/RemoteObjectTransfer.cpp
@@ -510,7 +510,7 @@ public:
 
 			if (RootOwnerServerId != FRemoteServerId::GetLocalServerId() && !RootOwnerServerId.IsAsset())
 			{
-				UE_LOG(LogRemoteObject, VeryVerbose, TEXT("FRemoteObjectTransferQueue: borrowed obj: %s (%s) owner server: %s (root: %s (s) root owner: %s)"),
+				UE_LOG(LogRemoteObject, VeryVerbose, TEXT("FRemoteObjectTransferQueue: borrowed obj: %s (%s) owner server: %s (root: %s (%s) root owner: %s)"),
 					*UsedObjectId.ToString(),
 					*UsedObject->GetName(),
 					*OwnerServerId.ToString(),
diff --git a/Engine/Source/Runtime/Engine/Private/Animation/AnimNode_StateMachine.cpp b/Engine/Source/Runtime/Engine/Private/Animation/AnimNode_StateMachine.cpp
--- a/Engine/Source/Runtime/Engine/Private/Animation/AnimNode_StateMachine.cpp
+++ b/Engine/Source/Runtime/Engine/Private/Animation/AnimNode_StateMachine.cpp
@@ -1483,7 +1483,7 @@ bool FAnimNode_StateMachine::RequestTransitionEvent(const FTransitionEvent& InTr
 		ensure(QueuedTransitionEvents.Num() <= MaxTransitionsRequests);
 		if (QueuedTransitionEvents.Num() == MaxTransitionsRequests)
 		{
-			UE_LOG(LogAnimTransitionRequests, Warning, TEXT("Transition request cap reached, dropping old requests (Machine: %s)"), *InTransitionEvent.EventName.ToString(), *GetMachineDescription()->MachineName.ToString());
+			UE_LOG(LogAnimTransitionRequests, Warning, TEXT("Transition '%s' request cap reached, dropping old requests (Machine: %s)"), *InTransitionEvent.EventName.ToString(), *GetMachineDescription()->MachineName.ToString());
 			QueuedTransitionEvents.Pop(EAllowShrinking::No);
 		}
 		QueuedTransitionEvents.Insert(InTransitionEvent, 0);
diff --git a/Engine/Source/Runtime/MovieScene/Private/MovieSceneMutualComponentInclusivity.cpp b/Engine/Source/Runtime/MovieScene/Private/MovieSceneMutualComponentInclusivity.cpp
--- a/Engine/Source/Runtime/MovieScene/Private/MovieSceneMutualComponentInclusivity.cpp
+++ b/Engine/Source/Runtime/MovieScene/Private/MovieSceneMutualComponentInclusivity.cpp
@@ -128,7 +128,7 @@ void FMutualInclusivityGraphCommandBuffer::CheckInvariants() const
 	auto TestShortCircuitIndex = [this](int32 Index)
 	{
 		checkf(Index != 0 && Index <= this->Commands.Num(),
-			TEXT("Invalid short circuit index %d for command buffer length"), Index, this->Commands.Num());
+			TEXT("Invalid short circuit index %d for command buffer length (%d)"), Index, this->Commands.Num());
 	};
 
 	for (int32 Index = 0; Index < Commands.Num(); ++Index)
diff --git a/Engine/Source/Runtime/Net/Iris/Private/Iris/ReplicationSystem/Filtering/ReplicationFiltering.cpp b/Engine/Source/Runtime/Net/Iris/Private/Iris/ReplicationSystem/Filtering/ReplicationFiltering.cpp
--- a/Engine/Source/Runtime/Net/Iris/Private/Iris/ReplicationSystem/Filtering/ReplicationFiltering.cpp
+++ b/Engine/Source/Runtime/Net/Iris/Private/Iris/ReplicationSystem/Filtering/ReplicationFiltering.cpp
@@ -1887,7 +1887,7 @@ void FReplicationFiltering::AddSubObjectFilter(FNetObjectGroupHandle GroupHandle
 
 	CreatePerSubObjectGroupFilterInfo(GroupIndex);
 
-	UE_LOG(LogIrisFiltering, Verbose, TEXT("ReplicationFiltering::AddSubObjectFilter Group: %s FilterStatus: DisallowReplication"), *Groups->GetGroupNameString(GroupHandle), GroupIndex);
+	UE_LOG(LogIrisFiltering, Verbose, TEXT("ReplicationFiltering::AddSubObjectFilter Group: %s Index: %u FilterStatus: DisallowReplication"), *Groups->GetGroupNameString(GroupHandle), GroupIndex);
 }
 
 void FReplicationFiltering::RemoveSubObjectFilter(FNetObjectGroupHandle GroupHandle)
diff --git a/Engine/Source/Runtime/PakFile/Private/IoDispatcherFileBackend.cpp b/Engine/Source/Runtime/PakFile/Private/IoDispatcherFileBackend.cpp
--- a/Engine/Source/Runtime/PakFile/Private/IoDispatcherFileBackend.cpp
+++ b/Engine/Source/Runtime/PakFile/Private/IoDispatcherFileBackend.cpp
@@ -947,13 +947,13 @@ TIoStatusOr<FIoContainerHeader> FFileIoStoreReader::ReadContainerHeader(bool bRe
 		if (ContainerHeader.SoftPackageReferencesSerialInfo.Offset < 0)
 		{
 			FIoStatus Status = FIoStatusBuilder(EIoErrorCode::ReadError)
-				<< FString::Printf(TEXT("Invalid soft package reference offset '" INT64_FMT "'"), ContainerHeader.SoftPackageReferencesSerialInfo.Offset);
+				<< FString::Printf(TEXT("Invalid soft package reference offset '%" INT64_FMT "'"), ContainerHeader.SoftPackageReferencesSerialInfo.Offset);
 			return Status;
 		}
 		if ((ContainerHeader.SoftPackageReferencesSerialInfo.Offset + ContainerHeader.SoftPackageReferencesSerialInfo.Size) > Ar.TotalSize())
 		{
 			FIoStatus Status = FIoStatusBuilder(EIoErrorCode::ReadError)
-				<< FString::Printf(TEXT("Soft package reference offset '" INT64_FMT "' and size '" INT64_FMT "' will seek past the end of archive size '" INT64_FMT "'"),
+				<< FString::Printf(TEXT("Soft package reference offset '%" INT64_FMT "' and size '%" INT64_FMT "' will seek past the end of archive size '%" INT64_FMT "'"),
 					ContainerHeader.SoftPackageReferencesSerialInfo.Offset,
 					ContainerHeader.SoftPackageReferencesSerialInfo.Size,
 					Ar.TotalSize());
diff --git a/Engine/Source/Runtime/RenderCore/Private/ShaderCodeLibrary.cpp b/Engine/Source/Runtime/RenderCore/Private/ShaderCodeLibrary.cpp
--- a/Engine/Source/Runtime/RenderCore/Private/ShaderCodeLibrary.cpp
+++ b/Engine/Source/Runtime/RenderCore/Private/ShaderCodeLibrary.cpp
@@ -3200,7 +3200,7 @@ public:
 					{
 						if (bOpenedAsMonolithic)
 						{
-							UE_LOG(LogShaderLibrary, Display, TEXT("Logical shader library '%s' has been created as a monolithic library"), *Name, Library->GetNumComponents());
+							UE_LOG(LogShaderLibrary, Display, TEXT("Logical shader library '%s' has been created as a monolithic library, components %d"), *Name, Library->GetNumComponents());
 						}
 						else
 						{
