Enabled reading soft references in new I/O dispatcher file backend.

#rnx
#Synced-CL 48408831

diff --git a/Engine/Source/Runtime/PakFile/Private/FileIoDispatcherBackend.cpp b/Engine/Source/Runtime/PakFile/Private/FileIoDispatcherBackend.cpp
--- a/Engine/Source/Runtime/PakFile/Private/FileIoDispatcherBackend.cpp
+++ b/Engine/Source/Runtime/PakFile/Private/FileIoDispatcherBackend.cpp
@@ -150,7 +150,7 @@ struct FContainerPartition
 ////////////////////////////////////////////////////////////////////////////////
 struct FContainer
 {
-	static TIoStatusOr<TUniquePtr<FContainer>>			Open(const TCHAR* Filename, int32 MountOrder, uint32 InstanceId, FIoContainerHeader& OutContainerHeader);
+	static TIoStatusOr<TUniquePtr<FContainer>>			Open(const TCHAR* Filename, int32 MountOrder, uint32 InstanceId, bool bReadsSoftRefs, FIoContainerHeader& OutContainerHeader);
 	FContainerPartition&								GetPartition(uint64 Offset, uint64& OutPartitionOffset, int32* OutIndex = nullptr);
 	uint64												GetAllocatedSize() const;
 
@@ -178,6 +178,7 @@ TIoStatusOr<TUniquePtr<FContainer>>	FContainer::Open(
 	const TCHAR* Filename,
 	int32 MountOrder,
 	uint32 InstanceId,
+	bool bReadSoftRefs,
 	FIoContainerHeader& OutContainerHeader)
 {
 	LLM_SCOPE(ELLMTag::AsyncLoading);
@@ -385,6 +386,27 @@ TIoStatusOr<TUniquePtr<FContainer>>	FContainer::Open(
 			return FIoStatus(FIoStatusBuilder(EIoErrorCode::FileOpenFailed)
 				<< TEXT("Failed to serialize container header"));
 		}
+
+		if (bReadSoftRefs)
+		{
+			if (OutContainerHeader.SoftPackageReferencesSerialInfo.Offset < 0)
+			{
+				return FIoStatusBuilder(EIoErrorCode::ReadError)
+					<< FString::Printf(TEXT("Invalid soft package reference offset '%" INT64_FMT "'"), OutContainerHeader.SoftPackageReferencesSerialInfo.Offset);
+			}
+
+			if ((OutContainerHeader.SoftPackageReferencesSerialInfo.Offset + OutContainerHeader.SoftPackageReferencesSerialInfo.Size) > Ar.TotalSize())
+			{
+				return FIoStatusBuilder(EIoErrorCode::ReadError)
+					<< FString::Printf(TEXT("Soft package reference offset '%" INT64_FMT "' and size '%" INT64_FMT "' will seek past the end of archive size '%" INT64_FMT "'"),
+						OutContainerHeader.SoftPackageReferencesSerialInfo.Offset,
+						OutContainerHeader.SoftPackageReferencesSerialInfo.Size,
+						Ar.TotalSize());
+			}
+
+			Ar.Seek(OutContainerHeader.SoftPackageReferencesSerialInfo.Offset);
+			Ar << OutContainerHeader.SoftPackageReferences;
+		}
 	}
 
 	return Container;
@@ -436,7 +458,7 @@ class FFileIoStore
 {
 public:
 	FChunkInfo						GetChunkInfo(const FIoChunkId& ChunkId) const;
-	TIoStatusOr<FIoContainerHeader>	Mount(const TCHAR* TocPath, int32 MountOrder);
+	TIoStatusOr<FIoContainerHeader>	Mount(const TCHAR* TocPath, int32 MountOrder, ETocMountOptions Options);
 	bool							Unmount(const TCHAR* TocPath);
 	void							ReopenAllFileHandles();
 	FRWLock&						GetLock() { return RwLock; }
@@ -462,19 +484,21 @@ FChunkInfo FFileIoStore::GetChunkInfo(const FIoChunkId& ChunkId) const
 }
 
 ////////////////////////////////////////////////////////////////////////////////
-TIoStatusOr<FIoContainerHeader>	FFileIoStore::Mount(const TCHAR* TocPath, int32 MountOrder)
+TIoStatusOr<FIoContainerHeader>	FFileIoStore::Mount(const TCHAR* TocPath, int32 MountOrder, ETocMountOptions Options)
 {
 	const uint32		InstanceId = ContainerInstanceId.fetch_add(1, std::memory_order_relaxed);
 	FIoContainerHeader	Hdr;
 
-	TIoStatusOr<FUniqueContainer> Status = FContainer::Open(TocPath, MountOrder, InstanceId, Hdr);
+	const bool bReadSoftRefs = EnumHasAnyFlags(Options, ETocMountOptions::WithSoftReferences);
+	TIoStatusOr<FUniqueContainer> Status = FContainer::Open(TocPath, MountOrder, InstanceId, bReadSoftRefs, Hdr);
 	if (Status.IsOk() == false)
 	{
 		return Status.Status();
 	}
 
-	FUniqueContainer Container	= Status.ConsumeValueOrDie();
-	int32 MountIndex			= INDEX_NONE;
+	FUniqueContainer Container				= Status.ConsumeValueOrDie();
+	const EIoContainerFlags ContainerFlags	= Container->ContainerFlags;
+	int32 MountIndex						= INDEX_NONE;
 	{
 		FWriteScopeLock _(RwLock);
 
@@ -493,7 +517,9 @@ TIoStatusOr<FIoContainerHeader>	FFileIoStore::Mount(const TCHAR* TocPath, int32
 		MountedContainers.Insert(MoveTemp(Container), MountIndex);
 	}
 
-	UE_LOG(LogIoStore, Log, TEXT("Mounted container '%s' at position %d"), TocPath, MountIndex);
+	UE_LOG(LogIoStore, Display, TEXT("Mounted container '%s', Id='%s', Order=%d, Slot=%d, Options=(%s), Flags=(%s), HasSoftRefs=%s"),
+		TocPath, *LexToString(Hdr.ContainerId), MountOrder, MountIndex, LexToString(Options), *LexToString(ContainerFlags),
+		Hdr.SoftPackageReferences.bContainsSoftPackageReferences ? TEXT("True") : TEXT("False"));
 
 	return Hdr;
 }
@@ -645,7 +671,7 @@ public:
 		int32 Order, 
 		const FGuid& EncryptionKeyGuid, 
 		const FAES::FAESKey& EncryptionKey, 
-		UE::IoStore::ETocMountOptions Options) override;
+		ETocMountOptions Options) override;
 	virtual bool							Unmount(const TCHAR* TocPath) override;
 	virtual void							ReopenAllFileHandles() override;
 
@@ -685,10 +711,10 @@ TIoStatusOr<FIoContainerHeader> FFileIoDispatcherBackend::Mount(
 	int32 Order, 
 	const FGuid& EncryptionKeyGuid, 
 	const FAES::FAESKey& EncryptionKey, 
-	UE::IoStore::ETocMountOptions Options)
+	ETocMountOptions Options)
 {
 	LLM_SCOPE_BYNAME(TEXT("FileSystem/FileIoStore"));
-	return IoStore.Mount(TocPath, Order);
+	return IoStore.Mount(TocPath, Order, Options);
 }
 
 bool FFileIoDispatcherBackend::Unmount(const TCHAR* TocPath)
