Log EncryptedHash in FPakSignatureFile::DecryptSignatureAndValidate when failing to decrypt pak signature to help track down cases of invalid pak/sig errors

#Synced-CL 48539628

diff --git a/Engine/Source/Runtime/PakFile/Public/IPlatformFilePak.h b/Engine/Source/Runtime/PakFile/Public/IPlatformFilePak.h
--- a/Engine/Source/Runtime/PakFile/Public/IPlatformFilePak.h
+++ b/Engine/Source/Runtime/PakFile/Public/IPlatformFilePak.h
@@ -2752,7 +2752,8 @@ struct FPakSignatureFile
 			}
 			else
 			{
-				UE_LOG(LogPakFile, Warning, TEXT("Pak signature table validation failed for '%s'! Failed to decrypt signature"), *InFilename);
+				FString EncryptedHashString = BytesToHex(EncryptedHash.GetData(), EncryptedHash.Num());
+				UE_LOG(LogPakFile, Warning, TEXT("Pak signature table validation failed for '%s'! Failed to decrypt signature. Signature EncryptedHash: (Bytes=%d, Data=%s)"), *InFilename, EncryptedHash.Num(), *EncryptedHashString);
 			}
 		}
 
