IoStore partition mmapping thread-safety

It is now thread-safe to mmap chunks from the same partition concurrently. For example, this occurs when ALT performs wide deserialization, while animations are memory-mapped (e.g. we have this on mobile).

#rnx

diff --git a/Engine/Source/Runtime/PakFile/Internal/IoDispatcherFileBackendTypes.h b/Engine/Source/Runtime/PakFile/Internal/IoDispatcherFileBackendTypes.h
--- a/Engine/Source/Runtime/PakFile/Internal/IoDispatcherFileBackendTypes.h
+++ b/Engine/Source/Runtime/PakFile/Internal/IoDispatcherFileBackendTypes.h
@@ -7,8 +7,11 @@
 #include "IO/IoDispatcherBackend.h"
 #include "IO/IoDispatcherFilesystemStats.h"
 #include "Async/MappedFileHandle.h"
+#include "Async/Mutex.h"
 #include "Modules/ModuleInterface.h"
 
+#include <atomic>
+
 #ifndef PLATFORM_IMPLEMENTS_IO
 #define PLATFORM_IMPLEMENTS_IO 0
 #endif
@@ -26,16 +29,24 @@ struct FFileIoStoreContainerFilePartition
 	FFileIoStoreContainerFilePartition() = default;
 	FFileIoStoreContainerFilePartition(FFileIoStoreContainerFilePartition&&) = delete;
 	FFileIoStoreContainerFilePartition(const FFileIoStoreContainerFilePartition&) = delete;
+	~FFileIoStoreContainerFilePartition();
 
 	FFileIoStoreContainerFilePartition& operator=(FFileIoStoreContainerFilePartition&&) = delete;
 	FFileIoStoreContainerFilePartition& operator=(const FFileIoStoreContainerFilePartition&) = delete;
 
+	// This will lazily map the file, if needed. It is thread-safe to borrow
+	// handles, but their lifetimes are tied to this partition.
+	IMappedFileHandle* BorrowMappedFileHandle();
+
+	FString FilePath;
 	uint64 FileHandle = 0;
 	uint64 FileSize = 0;
 	uint32 ContainerFileIndex = 0;
-	FString FilePath;
-	TUniquePtr<IMappedFileHandle> MappedFileHandle;
 	std::atomic<int32> StartedReadRequestsCount = 0;
+
+private:	
+	std::atomic<IMappedFileHandle*> MappedFileHandle = nullptr;
+	UE::FMutex MappedFileHandleMutex;
 };
 
 struct FFileIoStoreContainerFile
diff --git a/Engine/Source/Runtime/PakFile/Private/IoDispatcherFileBackend.cpp b/Engine/Source/Runtime/PakFile/Private/IoDispatcherFileBackend.cpp
--- a/Engine/Source/Runtime/PakFile/Private/IoDispatcherFileBackend.cpp
+++ b/Engine/Source/Runtime/PakFile/Private/IoDispatcherFileBackend.cpp
@@ -14,6 +14,7 @@
 #include "HAL/IConsoleManager.h"
 #include "Async/AsyncWork.h"
 #include "Async/MappedFileHandle.h"
+#include "Async/UniqueLock.h"
 #include "HAL/RunnableThread.h"
 #include "Misc/ScopeLock.h"
 #include "Misc/Paths.h"
@@ -636,6 +637,37 @@ int32 FFileIoStoreRequestQueue::HandleContainerUnmounted(const FFileIoStoreConta
 	return FailedRequestsCount;
 }
 
+FFileIoStoreContainerFilePartition::~FFileIoStoreContainerFilePartition()
+{
+	UE::TUniqueLock Lock(MappedFileHandleMutex);
+	delete MappedFileHandle.load(std::memory_order_relaxed);
+}
+
+IMappedFileHandle* FFileIoStoreContainerFilePartition::BorrowMappedFileHandle()
+{
+	IMappedFileHandle* Handle = MappedFileHandle.load(std::memory_order_acquire);
+	if (UNLIKELY(Handle == nullptr))
+	{
+		UE::TUniqueLock Lock(MappedFileHandleMutex);
+		Handle = MappedFileHandle.load(std::memory_order_relaxed);
+		if (UNLIKELY(Handle == nullptr))
+		{
+			IPlatformFile& Ipf = FPlatformFileManager::Get().GetPlatformFile();
+			FOpenMappedResult Result = Ipf.OpenMappedEx(*FilePath);
+			if (!Result.HasError())
+			{
+				Handle = Result.StealValue().Release();
+				MappedFileHandle.store(Handle, std::memory_order_release);
+			}
+			else
+			{
+				UE_LOG(LogIoDispatcher, Error, TEXT("OpenMappedEx failed on: %s"), *FilePath);
+			}
+		}
+	}
+	return Handle;
+}
+
 FFileIoStoreReader::FFileIoStoreReader(IPlatformFileIoStore& InPlatformImpl, FFileIoStoreStats& InStats)
 	: PlatformImpl(InPlatformImpl)
 	, Stats(InStats)
@@ -836,22 +868,20 @@ const FIoOffsetAndLength* FFileIoStoreReader::Resolve(const FIoChunkId& ChunkId)
 	return FindChunkInternal(ChunkId);
 }
 
-IMappedFileHandle* FFileIoStoreReader::GetMappedContainerFileHandle(uint64 TocOffset)
+TUniquePtr<IMappedFileHandle> FFileIoStoreReader::GetMappedContainerFileHandle(uint64 TocOffset)
 {
 	check(!bClosed);
 	int32 PartitionIndex = int32(TocOffset / ContainerFile.PartitionSize);
 	FFileIoStoreContainerFilePartition& Partition = ContainerFile.Partitions[PartitionIndex];
-	if (!Partition.MappedFileHandle &&
-		// Can't map encrypted files, compression should be disabled for bulk data when memory mapping is required
-		!EnumHasAnyFlags(ContainerFile.ContainerFlags, EIoContainerFlags::Encrypted))
+	IMappedFileHandle* MappedFileHandle = nullptr;
+	// Can't map encrypted files, compression should be disabled for bulk data when memory mapping is required
+	if (!EnumHasAnyFlags(ContainerFile.ContainerFlags, EIoContainerFlags::Encrypted))
 	{
-		IPlatformFile& Ipf = FPlatformFileManager::Get().GetPlatformFile();
-		FOpenMappedResult Result = Ipf.OpenMappedEx(*Partition.FilePath);
-		Partition.MappedFileHandle = Result.HasError() ? nullptr : Result.StealValue();
+		MappedFileHandle = Partition.BorrowMappedFileHandle();
 	}
 
 	check(Partition.FileSize > 0);
-	return new FMappedFileProxy(Partition.MappedFileHandle.Get(), Partition.FileSize, ContainerFile.PartitionSize);
+	return TUniquePtr<IMappedFileHandle>(new FMappedFileProxy(MappedFileHandle, Partition.FileSize, ContainerFile.PartitionSize));
 }
 
 TIoStatusOr<FIoContainerHeader> FFileIoStoreReader::ReadContainerHeader(bool bReadSoftRefs) const
@@ -1996,18 +2026,21 @@ TIoStatusOr<FIoMappedRegion> FFileIoStore::OpenMapped(const FIoChunkId& ChunkId,
 			// This is not required as mapped file region implementation ensure we map an aligned
 			// pointer/size, handling arbitrary offset/size.
 			// check(BlockOffset > 0 && IsAligned(BlockOffset, FPlatformProperties::GetMemoryMappingAlignment()));
-			IMappedFileHandle* MappedFileHandle = Reader->GetMappedContainerFileHandle(BlockOffset);
-			IMappedFileRegion* MappedFileRegion = MappedFileHandle->MapRegion(BlockOffset + Options.GetOffset(), ResolvedSize);
+			TUniquePtr<IMappedFileHandle> MappedFileHandle = Reader->GetMappedContainerFileHandle(BlockOffset);
+			const int64 MappingOffset = BlockOffset + Options.GetOffset();
+			IMappedFileRegion* MappedFileRegion = MappedFileHandle->MapRegion(MappingOffset, ResolvedSize);
 			if (MappedFileRegion != nullptr)
 			{
 				// MappedPtr could be inside of the page, depending on the offset. This
 				// is actually fine, unless client code cares about MappedPtr alignment, which is
 				// public API.
 				// check(IsAligned(MappedFileRegion->GetMappedPtr(), FPlatformMemory::GetStats().PageSize));
-				return FIoMappedRegion{ MappedFileHandle, MappedFileRegion };
+				return FIoMappedRegion{ MappedFileHandle.Release(), MappedFileRegion };
 			}
 			else
 			{
+				const FString& FilePath = Reader->GetContainerFile()->FilePath;
+				UE_LOG(LogIoDispatcher, Error, TEXT("MapRegion failed on: %s (0x%016" INT64_x_FMT ", size: %" UINT64_FMT ")"), *FilePath, MappingOffset, ResolvedSize);
 				return FIoStatus(EIoErrorCode::ReadError);
 			}
 		}
diff --git a/Engine/Source/Runtime/PakFile/Private/IoDispatcherFileBackend.h b/Engine/Source/Runtime/PakFile/Private/IoDispatcherFileBackend.h
--- a/Engine/Source/Runtime/PakFile/Private/IoDispatcherFileBackend.h
+++ b/Engine/Source/Runtime/PakFile/Private/IoDispatcherFileBackend.h
@@ -41,7 +41,7 @@ public:
 	const FIoOffsetAndLength* Resolve(const FIoChunkId& ChunkId) const;
 	FFileIoStoreContainerFile* GetContainerFile() { return &ContainerFile; }
 	const FFileIoStoreContainerFile* GetContainerFile() const { return &ContainerFile; }
-	IMappedFileHandle* GetMappedContainerFileHandle(uint64 TocOffset);
+	TUniquePtr<IMappedFileHandle> GetMappedContainerFileHandle(uint64 TocOffset);
 	const FIoContainerId& GetContainerId() const { return ContainerId; }
 	int32 GetOrder() const { return Order; }
 	bool IsEncrypted() const { return EnumHasAnyFlags(ContainerFile.ContainerFlags, EIoContainerFlags::Encrypted); }
