[PakFile]
Move LLMTags for mounting content to be under the Filesystem tag rather than AsyncLoading as this has more to do with the FilePackage/IOStore layers than the asset loading system (AsyncLoading2)

#rnx

#Synced-CL 50040740

diff --git a/Engine/Source/Runtime/PakFile/Private/FileIoDispatcherBackend.cpp b/Engine/Source/Runtime/PakFile/Private/FileIoDispatcherBackend.cpp
--- a/Engine/Source/Runtime/PakFile/Private/FileIoDispatcherBackend.cpp
+++ b/Engine/Source/Runtime/PakFile/Private/FileIoDispatcherBackend.cpp
@@ -181,7 +181,7 @@ TIoStatusOr<TUniquePtr<FContainer>>	FContainer::Open(
 	bool bReadSoftRefs,
 	FIoContainerHeader& OutContainerHeader)
 {
-	LLM_SCOPE(ELLMTag::AsyncLoading);
+	LLM_SCOPE_BYNAME(TEXT("Filesystem/IODispatcher"));
 	TRACE_CPUPROFILER_EVENT_SCOPE(ReadContainerHeader);
 
 	FString BaseFilePath = FString(FPathViews::GetBaseFilenameWithPath(Filename));
diff --git a/Engine/Source/Runtime/PakFile/Private/FilePackageStore.cpp b/Engine/Source/Runtime/PakFile/Private/FilePackageStore.cpp
--- a/Engine/Source/Runtime/PakFile/Private/FilePackageStore.cpp
+++ b/Engine/Source/Runtime/PakFile/Private/FilePackageStore.cpp
@@ -444,7 +444,7 @@ TConstArrayView<uint32> FFilePackageStoreBackend::GetSoftReferences(FPackageId P
 
 void FFilePackageStoreBackend::Mount(FIoContainerHeader* ContainerHeader, uint32 Order)
 {
-	LLM_SCOPE(ELLMTag::AsyncLoading);
+	LLM_SCOPE_BYNAME(TEXT("Filesystem/FilePackageStore"));
 	FWriteScopeLock _(EntriesLock);
 	MountedContainers.Add({ ContainerHeader, Order, NextSequence++ });
 	// Update() relies on stable sorting to efficiently drop unmounted package -> entry ranges, see OldIt
@@ -660,7 +660,7 @@ struct FFilePackageStoreEntryRef
 
 void FFilePackageStoreBackend::Update()
 {
-	LLM_SCOPE(ELLMTag::AsyncLoading);
+	LLM_SCOPE_BYNAME(TEXT("Filesystem/FilePackageStore"));
 	TRACE_CPUPROFILER_EVENT_SCOPE(UpdateFilePackageStore);
 
 	FScopeLock Lock(&UpdateLock);
diff --git a/Engine/Source/Runtime/PakFile/Private/IoDispatcherFileBackend.cpp b/Engine/Source/Runtime/PakFile/Private/IoDispatcherFileBackend.cpp
--- a/Engine/Source/Runtime/PakFile/Private/IoDispatcherFileBackend.cpp
+++ b/Engine/Source/Runtime/PakFile/Private/IoDispatcherFileBackend.cpp
@@ -886,7 +886,7 @@ TUniquePtr<IMappedFileHandle> FFileIoStoreReader::GetMappedContainerFileHandle(u
 
 TIoStatusOr<FIoContainerHeader> FFileIoStoreReader::ReadContainerHeader(bool bReadSoftRefs) const
 {
-	LLM_SCOPE(ELLMTag::AsyncLoading);
+	LLM_SCOPE_BYNAME(TEXT("Filesystem/FileIOStore"));
 	TRACE_CPUPROFILER_EVENT_SCOPE(ReadContainerHeader);
 	FIoChunkId HeaderChunkId = CreateIoChunkId(ContainerId.Value(), 0, EIoChunkType::ContainerHeader);
 	const FIoOffsetAndLength* OffsetAndLength = FindChunkInternal(HeaderChunkId);
