Use proxy file handle for memory mapped files in new I/O dispatcher file backend and remove check about alignment requirement.

#rnx

#Synced-CL 50811283

diff --git a/Engine/Source/Runtime/PakFile/Private/FileIoDispatcherBackend.cpp b/Engine/Source/Runtime/PakFile/Private/FileIoDispatcherBackend.cpp
--- a/Engine/Source/Runtime/PakFile/Private/FileIoDispatcherBackend.cpp
+++ b/Engine/Source/Runtime/PakFile/Private/FileIoDispatcherBackend.cpp
@@ -43,12 +43,14 @@ static FAutoConsoleVariableRef CVar_UnmountTimeOutSeconds(
 );
 
 ////////////////////////////////////////////////////////////////////////////////
-class FMappedFileProxy final : public IMappedFileHandle
+class FMappedFileProxy final
+	: public IMappedFileHandle
 {
 public:
-	FMappedFileProxy(IMappedFileHandle* InSharedMappedFileHandle, uint64 InSize)
-		: IMappedFileHandle(InSize)
+	FMappedFileProxy(IMappedFileHandle* InSharedMappedFileHandle, uint64 InFileSize, int64 InBlockOffset)
+		: IMappedFileHandle(InFileSize)
 		, SharedMappedFileHandle(InSharedMappedFileHandle)
+		, BlockOffset(InBlockOffset)
 	{
 	}
 
@@ -56,10 +58,11 @@ public:
 
 	virtual IMappedFileRegion* MapRegion(int64 Offset = 0, int64 BytesToMap = MAX_int64, FFileMappingFlags Flags = EMappedFileFlags::ENone) override
 	{
-		return SharedMappedFileHandle != nullptr ? SharedMappedFileHandle->MapRegion(Offset, BytesToMap, Flags) : nullptr;
+		return SharedMappedFileHandle != nullptr ? SharedMappedFileHandle->MapRegion(BlockOffset + Offset, BytesToMap, Flags) : nullptr;
 	}
 private:
 	IMappedFileHandle* SharedMappedFileHandle;
+	const int64 BlockOffset;
 };
 
 ////////////////////////////////////////////////////////////////////////////////
@@ -1015,7 +1018,7 @@ TIoStatusOr<FIoMappedRegion> FFileIoDispatcherBackend::OpenMapped(const FIoChunk
 		return FIoStatus(EIoErrorCode::NotFound);
 	}
 
-	const int64 ResolvedOffset	= ChunkInfo.Offset() + Options.GetOffset();
+	const int64 ResolvedOffset	= ChunkInfo.Offset();
 	const int64 ResolvedSize	= FMath::Min(Options.GetSize(), ChunkInfo.Size() - Options.GetOffset());
 	FContainer& Container		= ChunkInfo.GetContainer();
 	const int32 BlockIndex		= IntCastChecked<int32>(ResolvedOffset / Container.CompressionBlockSize);
@@ -1024,8 +1027,6 @@ TIoStatusOr<FIoMappedRegion> FFileIoDispatcherBackend::OpenMapped(const FIoChunk
 	uint64 BlockOffsetInPartition					= MAX_uint64;
 	FContainerPartition& Partition					= Container.GetPartition(Block.GetOffset(), BlockOffsetInPartition);
 
-	check(IsAligned(BlockOffsetInPartition, FPlatformProperties::GetMemoryMappingAlignment()));
-
 	if (Partition.MappedFileHandle.IsValid() == false)
 	{
 		IPlatformFile& Ipf = FPlatformFileManager::Get().GetPlatformFile();
@@ -1033,6 +1034,11 @@ TIoStatusOr<FIoMappedRegion> FFileIoDispatcherBackend::OpenMapped(const FIoChunk
 		{
 			Partition.MappedFileHandle = Result.StealValue();
 		}
+		else
+		{
+			UE_LOG(LogIoStore, Error, TEXT("Failed to memory map chunk, ChunkId='%s', Filename='%s', Reason='%s'"),
+				*LexToString(ChunkId), *Partition.Filename, *Result.GetError().GetMessage());
+		}
 	}
 
 	if (Partition.MappedFileHandle.IsValid() == false)
@@ -1045,10 +1051,15 @@ TIoStatusOr<FIoMappedRegion> FFileIoDispatcherBackend::OpenMapped(const FIoChunk
 	{
 		return FIoMappedRegion
 		{
-			.MappedFileHandle	= Partition.MappedFileHandle.Get(),
+			.MappedFileHandle	= new FMappedFileProxy(Partition.MappedFileHandle.Get(), MappedFileHandle.GetFileSize(), IntCastChecked<int64>(BlockOffsetInPartition)),
 			.MappedFileRegion	= MappedFileRegion
 		};
 	}
+	else
+	{
+		UE_LOG(LogIoStore, Error, TEXT("Failed to memory map region, ChunkId='%s', Filename='%s', Region=(%llu, %lld)"),
+			*LexToString(ChunkId), *Partition.Filename, BlockOffsetInPartition + Options.GetOffset(), ResolvedSize);
+	}
 
 	return FIoStatus(EIoErrorCode::FileOpenFailed);
 }
