[Engine] Clean package set immediately after use so we aren't holding on and referencing out of scope pointers

#tests preflights
#Synced-CL 51256550

diff --git a/Engine/Source/Runtime/PakFile/Private/FilePackageStore.cpp b/Engine/Source/Runtime/PakFile/Private/FilePackageStore.cpp
--- a/Engine/Source/Runtime/PakFile/Private/FilePackageStore.cpp
+++ b/Engine/Source/Runtime/PakFile/Private/FilePackageStore.cpp
@@ -734,7 +734,6 @@ void FFilePackageStoreBackend::Update()
 			{	
 				// Add an empty entry per mounted container to allow inverse mapping handle offset -> container,
 				// which in turn allows releasing FIoContainerHeader::PackageIds after first Update()
-				UniqueEntriesPerNewContainer.Reset();
 				UniqueEntriesPerNewContainer.Add({&EmptyEntry}, AddNewEntryData(EmptyEntry, nullptr));
 
 				TConstArrayView<FFilePackageStoreEntrySoftReferences> AllSoftReferences;
@@ -759,6 +758,9 @@ void FFilePackageStoreBackend::Update()
 					Pairs.Add({PackageId, Handle});
 				}
 
+				// Clear for next use
+				UniqueEntriesPerNewContainer.Reset();
+
 				MountedContainer.NumMountedPackages = NumNewPackages;
 				ContainerHeader->PackageIds.Empty();
 				ContainerHeader->SoftPackageReferences.PackageIndices.Empty();
