[Engine] Review feedback, update comment and add check to ensure re-use of container intentions are respected

#tests preflights
#Synced-CL 51344084

diff --git a/Engine/Source/Runtime/PakFile/Private/FilePackageStore.cpp b/Engine/Source/Runtime/PakFile/Private/FilePackageStore.cpp
--- a/Engine/Source/Runtime/PakFile/Private/FilePackageStore.cpp
+++ b/Engine/Source/Runtime/PakFile/Private/FilePackageStore.cpp
@@ -731,7 +731,9 @@ void FFilePackageStoreBackend::Update()
 				}
 			}
 			else if (int32 NumNewPackages = ContainerHeader->PackageIds.Num())
-			{	
+			{
+				check(UniqueEntriesPerNewContainer.IsEmpty());
+
 				// Add an empty entry per mounted container to allow inverse mapping handle offset -> container,
 				// which in turn allows releasing FIoContainerHeader::PackageIds after first Update()
 				UniqueEntriesPerNewContainer.Add({&EmptyEntry}, AddNewEntryData(EmptyEntry, nullptr));
@@ -758,7 +760,7 @@ void FFilePackageStoreBackend::Update()
 					Pairs.Add({PackageId, Handle});
 				}
 
-				// Clear for next use
+				// Clear container for next use before clearing the referenced data to avoid use-after-free
 				UniqueEntriesPerNewContainer.Reset();
 
 				MountedContainer.NumMountedPackages = NumNewPackages;
