Additional reference counting to track pending I/O requests.
- the current reference tracking of container partitions with active read requests does not include
asynchronous scatter tasks. This can in rare occasions cause crashes when unmounting container files.
- this change adds additional reference counting that includes scatter tasks and can be conditionally
activated with a new Cvar "iostore.UnmountTimeoutSeconds"

#rnx

diff --git a/Engine/Source/Runtime/PakFile/Internal/IoDispatcherFileBackendTypes.h b/Engine/Source/Runtime/PakFile/Internal/IoDispatcherFileBackendTypes.h
--- a/Engine/Source/Runtime/PakFile/Internal/IoDispatcherFileBackendTypes.h
+++ b/Engine/Source/Runtime/PakFile/Internal/IoDispatcherFileBackendTypes.h
@@ -52,11 +52,7 @@ private:
 struct FFileIoStoreContainerFile
 {
 	FFileIoStoreContainerFile() = default;
-	FFileIoStoreContainerFile(FFileIoStoreContainerFile&&) = default;
-	FFileIoStoreContainerFile(const FFileIoStoreContainerFile&) = delete;
-
-	FFileIoStoreContainerFile& operator=(FFileIoStoreContainerFile&&) = default;
-	FFileIoStoreContainerFile& operator=(const FFileIoStoreContainerFile&) = delete;
+	UE_NONCOPYABLE(FFileIoStoreContainerFile);
 
 	uint64 PartitionSize = 0;
 	uint64 CompressionBlockSize = 0;
@@ -69,6 +65,7 @@ struct FFileIoStoreContainerFile
 	TConstArrayView<FSHAHash> BlockSignatureTable;
 	TArray<FFileIoStoreContainerFilePartition> Partitions;
 	uint32 ContainerInstanceId = 0;
+	std::atomic_int32_t RefCount{0};
 
 	void GetPartitionAndOffset(uint64 TocOffset, FFileIoStoreContainerFilePartition*& OutPartition, uint64& OutOffset)
 	{
@@ -76,6 +73,21 @@ struct FFileIoStoreContainerFile
 		OutPartition = &Partitions[PartitionIndex];
 		OutOffset = TocOffset % PartitionSize;
 	}
+
+	void Empty()
+	{
+		PartitionSize = 0;
+		CompressionBlockSize = 0;
+		CompressionMethods.Empty();
+		CompressionBlocks = TConstArrayView<FIoStoreTocCompressedBlockEntry>();
+		FilePath.Empty();
+		EncryptionKeyGuid = FGuid();
+		EncryptionKey = FAES::FAESKey();
+		ContainerFlags = EIoContainerFlags::None;
+		BlockSignatureTable = TConstArrayView<FSHAHash>();
+		Partitions.Empty();
+		ContainerInstanceId = 0;
+	}
 };
 
 struct FFileIoStoreBuffer
diff --git a/Engine/Source/Runtime/PakFile/Private/IoDispatcherFileBackend.cpp b/Engine/Source/Runtime/PakFile/Private/IoDispatcherFileBackend.cpp
--- a/Engine/Source/Runtime/PakFile/Private/IoDispatcherFileBackend.cpp
+++ b/Engine/Source/Runtime/PakFile/Private/IoDispatcherFileBackend.cpp
@@ -32,6 +32,14 @@
 
 //PRAGMA_DISABLE_OPTIMIZATION
 
+float GFileIoStoreUnmountTimeoutSeconds = 0.0;
+static FAutoConsoleVariableRef CVar_FileIoStoreUnmountTimeoutSeconds(
+	TEXT("iostore.UnmountTimeoutSeconds"),
+	GFileIoStoreUnmountTimeoutSeconds,
+	TEXT("Number of seconds to wait for pending reads before unmounting containers."),
+	ECVF_Default
+);
+
 uint32 FFileIoStoreReadRequest::NextSequence = 0;
 #if CHECK_IO_STORE_READ_REQUEST_LIST_MEMBERSHIP
 uint32 FFileIoStoreReadRequestList::NextListCookie = 0;
@@ -786,7 +794,7 @@ FIoStatus FFileIoStoreReader::Close()
 
 	PerfectHashMap = FFileIoStoreReader::FPerfectHashMap();
 	TocImperfectHashMapFallback.Empty();
-	ContainerFile = FFileIoStoreContainerFile();
+	ContainerFile.Empty();
 	DataContainer = FIoStoreTocResourceStorage();
 	ContainerId = FIoContainerId();
 	Order = INDEX_NONE;
@@ -1017,7 +1025,7 @@ FFileIoStoreResolvedRequest::FFileIoStoreResolvedRequest(
 	, ResolvedSize(InResolvedSize)
 	, Priority(InPriority)
 {
-
+	ContainerFile->RefCount.fetch_add(1, std::memory_order_acq_rel);
 }
 
 void FFileIoStoreResolvedRequest::AddReadRequestLink(FFileIoStoreReadRequestLink* ReadRequestLink)
@@ -1433,6 +1441,28 @@ bool FFileIoStore::Unmount(const TCHAR* InTocPath)
 				}
 			}
 		}
+
+		if (GFileIoStoreUnmountTimeoutSeconds > 0.0)
+		{
+			// Perform a second reference count check that also includes pending scatter tasks
+			FFileIoStoreContainerFile& Container = *ReaderToUnmount->GetContainerFile();
+			int32 RefCount = Container.RefCount.load(std::memory_order_acquire);
+			UE_CLOG(RefCount > 0, LogIoDispatcher, Warning, TEXT("Waiting for %d request(s) before unmounting container '%s'"),
+				RefCount, InTocPath);
+
+			const double StartTime = FPlatformTime::Seconds();
+			while (RefCount > 0)
+			{
+				if (FPlatformTime::Seconds() - StartTime >= GFileIoStoreUnmountTimeoutSeconds)
+				{
+					UE_LOG(LogIoDispatcher, Warning, TEXT("Unmounting container '%s' with %d remaining requests(s) after waiting for %.2f seconds(s)."),
+						InTocPath, RefCount, GFileIoStoreUnmountTimeoutSeconds);
+					break;
+				}
+				FPlatformProcess::SleepNoStats(0.0);
+				RefCount = Container.RefCount.load(std::memory_order_acquire);
+			}
+		}
 	}
 	else
 	{
@@ -1684,6 +1714,8 @@ void FFileIoStore::CompleteDispatcherRequest(FFileIoStoreResolvedRequest* Resolv
 {
 	check(ResolvedRequest);
 	check(ResolvedRequest->DispatcherRequest);
+
+	ResolvedRequest->ContainerFile->RefCount.fetch_sub(1, std::memory_order_acq_rel);
 	FIoRequestImpl* DispatcherRequest = ResolvedRequest->DispatcherRequest;
 	ResolvedRequest->DispatcherRequest = nullptr;
 	if (ResolvedRequest->bFailed)
