Temp increase in time out value (0.0->5.0) to work around a potential memory stomp issue.

#rnx

#Synced-CL 51590581

diff --git a/Engine/Source/Runtime/PakFile/Private/IoDispatcherFileBackend.cpp b/Engine/Source/Runtime/PakFile/Private/IoDispatcherFileBackend.cpp
--- a/Engine/Source/Runtime/PakFile/Private/IoDispatcherFileBackend.cpp
+++ b/Engine/Source/Runtime/PakFile/Private/IoDispatcherFileBackend.cpp
@@ -32,7 +32,7 @@
 
 //PRAGMA_DISABLE_OPTIMIZATION
 
-float GFileIoStoreUnmountTimeoutSeconds = 0.0;
+float GFileIoStoreUnmountTimeoutSeconds = 5.0;
 static FAutoConsoleVariableRef CVar_FileIoStoreUnmountTimeoutSeconds(
 	TEXT("iostore.UnmountTimeoutSeconds"),
 	GFileIoStoreUnmountTimeoutSeconds,
