Move timeout CVar to namespace to avoid possible name clash in unity builds with old I/O dispatcher backend.

#rnx

#Synced-CL 51704681

diff --git a/Engine/Source/Runtime/PakFile/Private/FileIoDispatcherBackend.cpp b/Engine/Source/Runtime/PakFile/Private/FileIoDispatcherBackend.cpp
--- a/Engine/Source/Runtime/PakFile/Private/FileIoDispatcherBackend.cpp
+++ b/Engine/Source/Runtime/PakFile/Private/FileIoDispatcherBackend.cpp
@@ -35,13 +35,18 @@ namespace UE::IoStore
 {
 
 ////////////////////////////////////////////////////////////////////////////////
-float GFileIoStoreUnmountTimeOutSeconds = 10.0f;
+namespace Config
+{
+
+float FileIoStoreUnmountTimeOutSeconds = 10.0f;
 static FAutoConsoleVariableRef CVar_UnmountTimeOutSeconds(
 	TEXT("fileiostore.UnmountTimeOutSeconds"),
-	GFileIoStoreUnmountTimeOutSeconds,
+	FileIoStoreUnmountTimeOutSeconds,
 	TEXT("Max time to wait for pending I/O requests before unmounting a container.")
 );
 
+} // namespace Config
+
 ////////////////////////////////////////////////////////////////////////////////
 class FMappedFileProxy final
 	: public IMappedFileHandle
@@ -568,7 +573,7 @@ bool FFileIoStore::Unmount(const TCHAR* TocPath)
 		}
 
 		UE_LOG(LogIoStore, Log, TEXT("Waiting for read request(s) to finish before unmounting '%s.utoc'"), TocPath);
-		const double MaxWaitTimeSeconds	= GFileIoStoreUnmountTimeOutSeconds; 
+		const double MaxWaitTimeSeconds	= Config::FileIoStoreUnmountTimeOutSeconds;
 		const double StartTime			= FMath::Clamp(FPlatformTime::Seconds(), 5.0f, 30.0f);
 		while (ContainerToRemove->ActiveReadCount.load(std::memory_order_seq_cst) > 0)
 		{
