Make static analysis happy

#Synced-CL 51330535

diff --git a/Engine/Source/Runtime/PakFile/Private/IPlatformFilePak.cpp b/Engine/Source/Runtime/PakFile/Private/IPlatformFilePak.cpp
--- a/Engine/Source/Runtime/PakFile/Private/IPlatformFilePak.cpp
+++ b/Engine/Source/Runtime/PakFile/Private/IPlatformFilePak.cpp
@@ -4659,6 +4659,12 @@ template<typename OpenOptionsType>
 FOpenMappedResult TOpenMappedExImpl(FPakPlatformFile* pThis, const TCHAR* Filename, OpenOptionsType OpenOptions, int64 MaximumSize,
 	FOpenMappedResult(IPlatformFile::*LowLevelOpenMappedEx)(const TCHAR*, OpenOptionsType, int64))
 {
+#ifdef _MSC_VER
+	// MSVC's static analysis doesn't see that Result is initialized
+#pragma warning(push)
+#pragma warning(disable : 6001)
+#endif // _MSC_VER
+
 	if (!GMMIO_Enable)
 	{
 		return MakeError(TEXT("Memory mapped IO is not enabled on this platform"));
@@ -4703,6 +4709,10 @@ FOpenMappedResult TOpenMappedExImpl(FPakPlatformFile* pThis, const TCHAR* Filena
 		return (pThis->LowerLevel->*LowLevelOpenMappedEx)(Filename, OpenOptions, MaximumSize);
 	}
 	return MakeError(FString::Printf(TEXT("Failed to open mapped file '%s'"), Filename));
+
+#ifdef _MSC_VER
+#pragma warning(pop)
+#endif // _MSC_VER
 }
 
 FOpenMappedResult FPakPlatformFile::OpenMappedEx(const TCHAR* Filename, EOpenReadFlags OpenOptions, int64 MaximumSize)
