Use MSVC_PRAGMA to suppress warnings only for MSVC

#Synced-CL 51330535

diff --git a/Engine/Source/Runtime/PakFile/Private/IPlatformFilePak.cpp b/Engine/Source/Runtime/PakFile/Private/IPlatformFilePak.cpp
--- a/Engine/Source/Runtime/PakFile/Private/IPlatformFilePak.cpp
+++ b/Engine/Source/Runtime/PakFile/Private/IPlatformFilePak.cpp
@@ -4659,11 +4659,9 @@ template<typename OpenOptionsType>
 FOpenMappedResult TOpenMappedExImpl(FPakPlatformFile* pThis, const TCHAR* Filename, OpenOptionsType OpenOptions, int64 MaximumSize,
 	FOpenMappedResult(IPlatformFile::*LowLevelOpenMappedEx)(const TCHAR*, OpenOptionsType, int64))
 {
-#ifdef _MSC_VER
 	// MSVC's static analysis doesn't see that Result is initialized
-#pragma warning(push)
-#pragma warning(disable : 6001)
-#endif // _MSC_VER
+MSVC_PRAGMA(warning(push))
+MSVC_PRAGMA(warning(disable : 6001))
 
 	if (!GMMIO_Enable)
 	{
@@ -4710,9 +4708,8 @@ FOpenMappedResult TOpenMappedExImpl(FPakPlatformFile* pThis, const TCHAR* Filena
 	}
 	return MakeError(FString::Printf(TEXT("Failed to open mapped file '%s'"), Filename));
 
-#ifdef _MSC_VER
-#pragma warning(pop)
-#endif // _MSC_VER
+
+MSVC_PRAGMA(warning(pop))
 }
 
 FOpenMappedResult FPakPlatformFile::OpenMappedEx(const TCHAR* Filename, EOpenReadFlags OpenOptions, int64 MaximumSize)
