Unload I/O store containers asynchronously to better handle pending I/O request(s).

- Added a new CVar iostore.AsyncUnmountEnabled=false that when enabled unloads I/O store container file(s) asynchronously from the unmount call
- Changed unmount to always wait for some period of time since having a zero timeout value is not valid

#rnx

#Synced-CL 51936612

diff --git a/Engine/Source/Runtime/PakFile/Private/IoDispatcherFileBackend.cpp b/Engine/Source/Runtime/PakFile/Private/IoDispatcherFileBackend.cpp
--- a/Engine/Source/Runtime/PakFile/Private/IoDispatcherFileBackend.cpp
+++ b/Engine/Source/Runtime/PakFile/Private/IoDispatcherFileBackend.cpp
@@ -36,10 +36,17 @@ float GFileIoStoreUnmountTimeoutSeconds = 5.0;
 static FAutoConsoleVariableRef CVar_FileIoStoreUnmountTimeoutSeconds(
 	TEXT("iostore.UnmountTimeoutSeconds"),
 	GFileIoStoreUnmountTimeoutSeconds,
-	TEXT("Number of seconds to wait for pending reads before unmounting containers."),
+	TEXT("Number of seconds (1-10) to wait for pending reads before unmounting containers."),
 	ECVF_Default
 );
 
+bool GFileIoStoreAsyncUnmountEnabled = false;
+static FAutoConsoleVariableRef CVar_FileIoStoreAsyncUnmountEnabled(
+	TEXT("iostore.AsyncUnmountEnabled"),
+	GFileIoStoreAsyncUnmountEnabled,
+	TEXT("When enabled containers will be unloaded from memory asynchronously from unmount.")
+);
+
 uint32 FFileIoStoreReadRequest::NextSequence = 0;
 #if CHECK_IO_STORE_READ_REQUEST_LIST_MEMBERSHIP
 uint32 FFileIoStoreReadRequestList::NextListCookie = 0;
@@ -1266,6 +1273,79 @@ int64 FFileIoStoreRequestTracker::GetLiveReadRequestsCount() const
 	return RequestAllocator.GetLiveReadRequestsCount();
 }
 
+FFileIoStoreReadersToUnload::~FFileIoStoreReadersToUnload()
+{
+	TRACE_CPUPROFILER_EVENT_SCOPE(WaitOnReferencedContainers);
+
+	const double MaxWaitTime = FMath::Clamp(GFileIoStoreUnmountTimeoutSeconds, 1.0, 10.0);
+	const double StartTime = FPlatformTime::Seconds();
+
+	for (;;)
+	{
+		int32 Pending = 0;
+		{
+			UE::TUniqueLock Lock(Mutex);
+			for (TUniquePtr<FFileIoStoreReader>& Reader : Readers)
+			{
+				if (Reader->GetContainerFile()->RefCount.load(std::memory_order_acquire) > 0)
+				{
+					Pending++;
+				}
+			}
+		}
+
+		if (Pending == 0)
+		{
+			break;
+		}
+
+		if (double WaitTime = FPlatformTime::Seconds() - StartTime; WaitTime > MaxWaitTime)
+		{
+			UE_LOG(LogIoDispatcher, Error, TEXT("Shutting down file I/O dispatcher backend with %d referenced container(s)."), Pending);
+			break;
+		}
+
+		FPlatformProcess::SleepNoStats(0.01f);
+	}
+}
+
+void FFileIoStoreReadersToUnload::Add(TUniquePtr<FFileIoStoreReader>&& Reader)
+{
+	UE::TUniqueLock Lock(Mutex);
+	Readers.Add(MoveTemp(Reader));
+}
+
+void FFileIoStoreReadersToUnload::TryUnload()
+{
+	TArray<TUniquePtr<FFileIoStoreReader>> ToUnload;
+	{
+		UE::TUniqueLock Lock(Mutex);
+		for (auto It = Readers.CreateIterator(); It; ++It)
+		{
+			TUniquePtr<FFileIoStoreReader>& Reader	= *It;
+			FFileIoStoreContainerFile& Container	= *Reader->GetContainerFile();
+			const int32 RefCount					= Container.RefCount.load(std::memory_order_acquire);
+
+			if (RefCount > 0)
+			{
+				UE_LOG(LogIoDispatcher, Warning, TEXT("Waiting on %d pending I/O request before unloading container '%s'"),
+					RefCount, *Container.FilePath);
+			}
+			else
+			{
+				check(RefCount == 0);
+				ToUnload.Add(MoveTemp(Reader));
+				It.RemoveCurrent();
+			}
+		}
+	}
+
+	for (TUniquePtr<FFileIoStoreReader>& Reader : ToUnload)
+	{
+		UE_LOG(LogIoDispatcher, Display, TEXT("Unloading container '%s'"), *Reader->GetContainerFile()->FilePath);
+	}
+}
+
 FFileIoStore::FFileIoStore(TUniquePtr<IPlatformFileIoStore>&& InPlatformImpl)
 	: BlockCache(Stats)
 	, BufferAllocator(Stats)
@@ -1428,6 +1508,13 @@ bool FFileIoStore::Unmount(const TCHAR* InTocPath)
 
 		UE_CLOG(FailedRequestsCount > 0, LogIoDispatcher, Warning, TEXT("Marking %d queued requests from unmounted container as failed"), FailedRequestsCount);
 
+		if (GFileIoStoreAsyncUnmountEnabled)
+		{
+			IoStoreReadersToUnload.Add(MoveTemp(ReaderToUnmount));
+			BackendContext->WakeUpDispatcherThreadDelegate.Execute();
+			return true;
+		}
+
 		bool bHasWarned = false;
 		for (const FFileIoStoreContainerFilePartition& Partition : ReaderToUnmount->GetContainerFile()->Partitions)
 		{
@@ -1442,27 +1529,27 @@ bool FFileIoStore::Unmount(const TCHAR* InTocPath)
 			}
 		}
 
-		if (GFileIoStoreUnmountTimeoutSeconds > 0.0)
-		{
-			// Perform a second reference count check that also includes pending scatter tasks
-			FFileIoStoreContainerFile& Container = *ReaderToUnmount->GetContainerFile();
-			int32 RefCount = Container.RefCount.load(std::memory_order_acquire);
-			UE_CLOG(RefCount > 0, LogIoDispatcher, Warning, TEXT("Waiting for %d request(s) before unmounting container '%s'"),
-				RefCount, InTocPath);
+		// Perform a second reference count check that also includes pending scatter tasks
+		FFileIoStoreContainerFile& Container = *ReaderToUnmount->GetContainerFile();
+		int32 RefCount = Container.RefCount.load(std::memory_order_acquire);
+		UE_CLOG(RefCount > 0, LogIoDispatcher, Warning, TEXT("Waiting for %d request(s) before unmounting container '%s'"),
+			RefCount, InTocPath);
 
-			const double StartTime = FPlatformTime::Seconds();
-			while (RefCount > 0)
+		const double MaxWaitTime = FMath::Clamp(GFileIoStoreUnmountTimeoutSeconds, 1.0, 10.0);
+		const double StartTime = FPlatformTime::Seconds();
+		while (RefCount > 0)
+		{
+			if (FPlatformTime::Seconds() - StartTime > MaxWaitTime)
 			{
-				if (FPlatformTime::Seconds() - StartTime >= GFileIoStoreUnmountTimeoutSeconds)
-				{
-					UE_LOG(LogIoDispatcher, Warning, TEXT("Unmounting container '%s' with %d remaining requests(s) after waiting for %.2f seconds(s)."),
-						InTocPath, RefCount, GFileIoStoreUnmountTimeoutSeconds);
-					break;
-				}
-				FPlatformProcess::SleepNoStats(0.0);
-				RefCount = Container.RefCount.load(std::memory_order_acquire);
+				UE_LOG(LogIoDispatcher, Warning, TEXT("Unmounting container '%s' with %d remaining requests(s) after waiting for %.2f seconds(s)."),
+					InTocPath, RefCount, GFileIoStoreUnmountTimeoutSeconds);
+				break;
 			}
+			FPlatformProcess::SleepNoStats(0.0);
+			RefCount = Container.RefCount.load(std::memory_order_acquire);
 		}
+
+		return true;
 	}
 	else
 	{
@@ -2025,6 +2112,8 @@ FIoRequestImpl* FFileIoStore::GetCompletedIoRequests()
 		ReadyForDecompressionTail = nullptr;
 	}
 
+	IoStoreReadersToUnload.TryUnload();
+
 	FIoRequestImpl* Result = CompletedRequestsHead;
 	CompletedRequestsHead = CompletedRequestsTail = nullptr;
 	return Result;
diff --git a/Engine/Source/Runtime/PakFile/Private/IoDispatcherFileBackend.h b/Engine/Source/Runtime/PakFile/Private/IoDispatcherFileBackend.h
--- a/Engine/Source/Runtime/PakFile/Private/IoDispatcherFileBackend.h
+++ b/Engine/Source/Runtime/PakFile/Private/IoDispatcherFileBackend.h
@@ -4,6 +4,8 @@
 
 #include "FileIoDispatcherBackend.h"
 
+#include "Async/Mutex.h"
+#include "Async/UniqueLock.h"
 #include "IoDispatcherFileBackendTypes.h"
 #include "IO/IoDispatcher.h"
 #include "IO/IoStore.h"
@@ -16,6 +18,8 @@
 #include "Misc/AES.h"
 #include "GenericPlatform/GenericPlatformFile.h"
 
+#include <atomic>
+
 class IMappedFileHandle;
 
 struct FFileIoStoreCompressionContext
@@ -103,6 +107,19 @@ private:
 	TMap<FFileIoStoreBlockKey, FFileIoStoreReadRequest*> RawBlocksMap;
 };
 
+class FFileIoStoreReadersToUnload
+{
+public:
+	FFileIoStoreReadersToUnload() = default;
+	~FFileIoStoreReadersToUnload();
+	void Add(TUniquePtr<FFileIoStoreReader>&& Reader);
+	void TryUnload();
+
+private:
+	TArray<TUniquePtr<FFileIoStoreReader>> Readers;
+	UE::FMutex Mutex;
+};
+
 class FFileIoStore final
 	: public FRunnable
 	, public UE::IoStore::IFileIoDispatcherBackend
@@ -157,6 +174,7 @@ private:
 	FFileIoStoreRequestQueue RequestQueue;
 	FFileIoStoreRequestTracker RequestTracker;
 	TUniquePtr<IPlatformFileIoStore> PlatformImpl;
+	FFileIoStoreReadersToUnload IoStoreReadersToUnload;
 	FRunnableThread* Thread = nullptr;
 	bool bIsMultithreaded;
 	TAtomic<bool> bStopRequested{ false };
