WindowsPlatformFile: Fix IterateDirectory to honor the contract to return false if the input directory does not exist. Keep the behavior that it returns true for directories that exist but are hidden or otherwise uniterateable.

PlatformFileSandboxWrapper: Fix IterateDirectory to fall through to iterate the non-sandbox file even if IterateDirectory on the sandbox file returned false because it did not exist.

This change fixes CookOnTheFly on linux when running in a workspace with an empty Saved\Cooked\<Platform> directory; it was previously not reporting files under Root\Engine because their corresponding path under Saved\Cooked\<Platform> did not exist.

#Synced-CL 52605363

diff --git a/Engine/Source/Runtime/Core/Private/Windows/WindowsPlatformFile.cpp b/Engine/Source/Runtime/Core/Private/Windows/WindowsPlatformFile.cpp
--- a/Engine/Source/Runtime/Core/Private/Windows/WindowsPlatformFile.cpp
+++ b/Engine/Source/Runtime/Core/Private/Windows/WindowsPlatformFile.cpp
@@ -1893,11 +1893,11 @@ bool FWindowsPlatformFile::IterateDirectoryCommon(const TCHAR* Directory, const
 {
 	TRACE_CPUPROFILER_EVENT_SCOPE(FWindowsPlatformFile_IterateDirectoryCommon_WithCallback);
 
-	bool bResult = true;
 	WIN32_FIND_DATAW Data;
 	HANDLE Handle = FindFirstFileWithWildcard(Directory, /* Out */ Data);
 	if (Handle != INVALID_HANDLE_VALUE)
 	{
+		bool bResult = true;
 		do
 		{
 			if (FCString::Strcmp(Data.cFileName, TEXT(".")) && FCString::Strcmp(Data.cFileName, TEXT("..")))
@@ -1906,8 +1906,26 @@ bool FWindowsPlatformFile::IterateDirectoryCommon(const TCHAR* Directory, const
 			}
 		} while (bResult && FindNextFileW(Handle, &Data));
 		FindClose(Handle);
+		return bResult;
+	}
+	else
+	{
+		// FindFirstFileWithWildcard can return an invalid handle even for directories that exist on disk, due to flags
+		// on the directory that instruct the operating system to hide them from various listings. e.g.
+		//     FILE_ATTRIBUTE_HIDDEN: The file or directory is hidden. It is not included in an ordinary directory listing.
+		//     FILE_ATTRIBUTE_REPARSE_POINT: A file or directory that has an associated reparse point, or a file that is a symbolic link.
+		//     FILE_ATTRIBUTE_SYSTEM: A file or directory that the operating system uses a part of, or uses exclusively.
+		// 
+		// But unfortunately, FindNextFileW on the parent directory will still report those hidden directories, and we
+		// therefore report them to our caller. IPlatformFile::IterateDirectoryRecursively on the parent
+		// directory reasonably expects that for these directories we report, calling IterateDirectory on them does
+		// not return false (unless the Visitor returns false).
+		//
+		// To workaround this issue, we return true for these hidden directories.
+		// But we still need to return false if the directory does not exist.
+		// DirectoryExists uses the function GetFileAttributes to decide whether the directory exists.
+		return DirectoryExists(Directory);
 	}
-	return bResult;
 }
 
 IPlatformFile& IPlatformFile::GetPlatformPhysical()
diff --git a/Engine/Source/Runtime/PakFile/Private/IPlatformFilePak.cpp b/Engine/Source/Runtime/PakFile/Private/IPlatformFilePak.cpp
--- a/Engine/Source/Runtime/PakFile/Private/IPlatformFilePak.cpp
+++ b/Engine/Source/Runtime/PakFile/Private/IPlatformFilePak.cpp
@@ -330,6 +330,8 @@ public:
 	virtual bool ShouldVisitLeafPathname(FStringView LeafNormalizedPathname) = 0;
 	virtual bool Visit(const FString& Filename, const FString& NormalizedFilename, bool bIsDir, FPakFile& PakFile) = 0;
 	// No need for CallShouldVisitAndVisit because we call ShouldVisitLeafPathname separately in all cases
+
+	bool bKeepIterating = true;
 };
 
 /** FPakFileDirectoryVisitorBase for a FDirectoryVisitor. */
@@ -347,7 +349,8 @@ public:
 	virtual bool Visit(const FString& Filename, const FString& NormalizedFilename,
 		bool bIsDir, FPakFile& PakFile) override
 	{
-		return Inner.Visit(*NormalizedFilename, bIsDir);
+		bKeepIterating = bKeepIterating && Inner.Visit(*NormalizedFilename, bIsDir);
+		return bKeepIterating;
 	}
 
 	IPlatformFile::FDirectoryVisitor& Inner;
@@ -415,7 +418,8 @@ public:
 			true	// IsReadOnly
 		);
 
-		return Inner.Visit(*NormalizedFilename, StatData);
+		bKeepIterating = bKeepIterating && Inner.Visit(*NormalizedFilename, StatData);
+		return bKeepIterating;
 	}
 
 	FPakPlatformFile& PlatformFile;
@@ -4933,22 +4937,33 @@ bool FPakPlatformFile::IterateDirectoryStatInternal(const TCHAR* Directory,
 	FPakFileDirectoryStatVisitor PakVisitor(*this, Visitor);
 	TSet<FString> FilesVisitedInPak;
 	bool Result = IterateDirectoryInPakFiles(Directory, PakVisitor, bRecursive, FilesVisitedInPak);
-	if (Result && LowerLevel->DirectoryExists(Directory))
+	if (!PakVisitor.bKeepIterating)
 	{
-		// Iterate inner filesystem but don't visit any files that were found in the Paks
-		FPreventDuplicatesStatVisitor PreventDuplicatesVisitor(Visitor, FilesVisitedInPak);
-		IPlatformFile::FDirectoryStatVisitor& LowerLevelVisitor(
-			// For performance, skip using PreventDuplicatedVisitor if there were no hits in pak
-			FilesVisitedInPak.Num() ? PreventDuplicatesVisitor : Visitor);
-		if (bRecursive)
-		{
-			Result = LowerLevel->IterateDirectoryStatRecursively(Directory, LowerLevelVisitor);
-		}
-		else
-		{
-			Result = LowerLevel->IterateDirectoryStat(Directory, LowerLevelVisitor);
-		}
+		return false;
 	}
+	bool bExistsInPakFiles = Result;
+	if (!LowerLevel->DirectoryExists(Directory))
+	{
+		return bExistsInPakFiles;
+	}
+
+	// Iterate inner filesystem but don't visit any files that were found in the Paks
+	FPreventDuplicatesStatVisitor PreventDuplicatesVisitor(Visitor, FilesVisitedInPak);
+	IPlatformFile::FDirectoryStatVisitor& LowerLevelVisitor(
+		// For performance, skip using PreventDuplicatedVisitor if there were no hits in pak
+		FilesVisitedInPak.Num() ? PreventDuplicatesVisitor : Visitor);
+	if (bRecursive)
+	{
+		Result = LowerLevel->IterateDirectoryStatRecursively(Directory, LowerLevelVisitor);
+	}
+	else
+	{
+		Result = LowerLevel->IterateDirectoryStat(Directory, LowerLevelVisitor);
+	}
+	// LowerLevel bDirectoryExists == true, so
+	// 1) Our contract to return (bKeepIterating && bDirectoryExists) therefore reduces to return (bKeepIterating).
+	// 2) LowerLevel IterateDirectoryStat return value also reduces to bKeepIterating.
+	// So just return Result.
 	return Result;
 }
 
@@ -5141,31 +5156,41 @@ bool FPakPlatformFile::IterateDirectoryInternal(const TCHAR* Directory,
 	FPakFileDirectoryVisitor PakVisitor(Visitor);
 	TSet<FString> FilesVisitedInPak;
 	bool Result = IterateDirectoryInPakFiles(Directory, PakVisitor, bRecursive, FilesVisitedInPak);
-	if (Result && LowerLevel->DirectoryExists(Directory))
+	if (!PakVisitor.bKeepIterating)
 	{
-		// Iterate inner filesystem but don't visit any files that were found in the Paks
-		FPreventDuplicatesVisitor PreventDuplicatesVisitor(Visitor, FilesVisitedInPak);
-		IPlatformFile::FDirectoryVisitor& LowerLevelVisitor(
-			// For performance, skip using PreventDuplicatedVisitor if there were no hits in pak
-			FilesVisitedInPak.Num() ? PreventDuplicatesVisitor : Visitor
-		);
-		if (bRecursive)
-		{
-			Result = LowerLevel->IterateDirectoryRecursively(Directory, LowerLevelVisitor);
-		}
-		else
-		{
-			Result = LowerLevel->IterateDirectory(Directory, LowerLevelVisitor);
-		}
+		return false;
+	}
+	bool bExistsInPakFiles = Result;
+	if (!LowerLevel->DirectoryExists(Directory))
+	{
+		return bExistsInPakFiles;
 	}
+
+	// Iterate inner filesystem but don't visit any files that were found in the Paks
+	FPreventDuplicatesVisitor PreventDuplicatesVisitor(Visitor, FilesVisitedInPak);
+	IPlatformFile::FDirectoryVisitor& LowerLevelVisitor(
+		// For performance, skip using PreventDuplicatedVisitor if there were no hits in pak
+		FilesVisitedInPak.Num() ? PreventDuplicatesVisitor : Visitor
+	);
+	if (bRecursive)
+	{
+		Result = LowerLevel->IterateDirectoryRecursively(Directory, LowerLevelVisitor);
+	}
+	else
+	{
+		Result = LowerLevel->IterateDirectory(Directory, LowerLevelVisitor);
+	}
+
+	// LowerLevel bDirectoryExists == true, so
+	// 1) Our contract to return (bKeepIterating && bDirectoryExists) therefore reduces to return (bKeepIterating).
+	// 2) LowerLevel IterateDirectory return value also reduces to bKeepIterating.
+	// So just return Result.
 	return Result;
 }
 
 bool FPakPlatformFile::IterateDirectoryInPakFiles(const TCHAR* Directory,
 	UE::PakFile::Private::FPakFileDirectoryVisitorBase& Visitor, bool bRecursive, TSet<FString>& FilesVisitedInPak)
 {
-	bool Result = true;
-
 	TArray<FPakListEntry> Paks;
 	FString StandardDirectory = Directory;
 	FPaths::MakeStandardFilename(StandardDirectory);
@@ -5183,14 +5208,28 @@ bool FPakPlatformFile::IterateDirectoryInPakFiles(const TCHAR* Directory,
 		GetMountedPaks(Paks);
 	}
 
+	bool bPathExistsInPak = false;
 	// Iterate pak files first
 	FString NormalizationBuffer;
 	TSet<FString> FilesVisitedInThisPak;
-	auto ShouldVisit = [&Visitor](FStringView UnnormalizedPath)
-	{
+	auto ShouldVisit = [&Visitor, &bPathExistsInPak](FStringView UnnormalizedPath)
+	{
+		// This calculation of bPathExistsInPak is currently correct but depends on bIncludeFiles=true and
+		// bIncludeFolders=true both being set. We need to report whether the queried directory exists in the pak
+		// file at all, but we only get a call to our ShouldVisit function here if it has files or folders in it.
+		// For pak files it is equivalent to say that PathExists <=> ItHasFilesOrFoldersInIt, so our equivalency
+		// that ShouldVisitIsCalled <=> bPathExistsInPak is valid.
+		// But if we were to set bIncludeFiles=false or bIncludeFolders=false, we would skip being called for
+		// in some cases even when PathExists, and that would break the equivalency.
+		// We rely on this behavior rather than changing FindPrunedFilesAtPathInternal to return whether the path
+		// exists because determining whether the path exists is expensive when FindPrunedFilesAtPathInternal is called
+		// for other purposes that don't care about bPathExistsInPak and want to skip the enumeration of files
+		// or folders.
+		bPathExistsInPak = true;
 		FStringView NormalizedPath = UE::String::RemoveFromEnd(UnnormalizedPath, TEXTVIEW("/"));
 		return Visitor.ShouldVisitLeafPathname(FPathViews::GetCleanFilename(NormalizedPath));
 	};
+
 	for (int32 PakIndex = 0; PakIndex < Paks.Num(); PakIndex++)
 	{
 		FPakFile& PakFile = *Paks[PakIndex].PakFile;
@@ -5200,7 +5239,7 @@ bool FPakPlatformFile::IterateDirectoryInPakFiles(const TCHAR* Directory,
 
 		FilesVisitedInThisPak.Reset();
 		PakFile.FindPrunedFilesAtPathInternal(*StandardDirectory, FilesVisitedInThisPak, FPakFile::FVisitFilter(ShouldVisit, bIncludeFiles, bIncludeFolders, bRecursive));
-		for (TSet<FString>::TConstIterator SetIt(FilesVisitedInThisPak); SetIt && Result; ++SetIt)
+		for (TSet<FString>::TConstIterator SetIt(FilesVisitedInThisPak); SetIt; ++SetIt)
 		{
 			const FString& Filename = *SetIt;
 			bool bIsDir = Filename.Len() && Filename[Filename.Len() - 1] == '/';
@@ -5218,11 +5257,14 @@ bool FPakPlatformFile::IterateDirectoryInPakFiles(const TCHAR* Directory,
 			if (!FilesVisitedInPak.Contains(*NormalizedFilename))
 			{
 				FilesVisitedInPak.Add(*NormalizedFilename);
-				Result = Visitor.Visit(Filename, *NormalizedFilename, bIsDir, PakFile) && Result;
+				if (!Visitor.Visit(Filename, *NormalizedFilename, bIsDir, PakFile))
+				{
+					return false;
+				}
 			}
 		}
 	}
-	return Result;
+	return bPathExistsInPak;
 }
 
 bool FPakPlatformFile::IterateDirectoryRecursively(const TCHAR* Directory, IPlatformFile::FDirectoryVisitor& Visitor)
diff --git a/Engine/Source/Runtime/SandboxFile/Private/IPlatformFileSandboxWrapper.cpp b/Engine/Source/Runtime/SandboxFile/Private/IPlatformFileSandboxWrapper.cpp
--- a/Engine/Source/Runtime/SandboxFile/Private/IPlatformFileSandboxWrapper.cpp
+++ b/Engine/Source/Runtime/SandboxFile/Private/IPlatformFileSandboxWrapper.cpp
@@ -719,7 +719,8 @@ public:
 	FDirectoryVisitor& Visitor;
 	FSandboxPlatformFile& SandboxFile;
 	TSet<FString> VisitedSandboxFiles;
-	bool bIsRecursive;
+	bool bIsRecursive = false;
+	bool bKeepIterating = true;
 
 	FSandboxVisitor(FDirectoryVisitor& InVisitor, FSandboxPlatformFile& InSandboxFile, bool bInIsRecursive)
 		: Visitor(InVisitor)
@@ -745,12 +746,20 @@ public:
 			if (NormalizedFilename == SandboxFile.InjectedSourceDirectoryParent)
 			{
 				// "fake" the injected directory so everything falls into place
-				Visitor.CallShouldVisitAndVisit(*SandboxFile.InjectedSourceDirectory, true);
+				bKeepIterating = bKeepIterating && Visitor.CallShouldVisitAndVisit(*SandboxFile.InjectedSourceDirectory, true);
+				if (!bKeepIterating)
+				{
+					return false;
+				}
 
 				// for recursive visiting, we need to recurse into the forced directory
 				if (bIsRecursive)
 				{
 					SandboxFile.LowerLevel->IterateDirectoryRecursively(*SandboxFile.InjectedTargetDirectory, *this);
+					if (!bKeepIterating)
+					{
+						return false;
+					}
 				}
 			}
 
@@ -792,14 +801,14 @@ public:
 
 		if (CanVisit)
 		{
-			bool Result = Visitor.CallShouldVisitAndVisit(*LocalFilename, bIsDirectory);
-			return Result;
-		}
-		else
-		{
-			// Continue iterating.
-			return true;
+			bKeepIterating = bKeepIterating && Visitor.CallShouldVisitAndVisit(*LocalFilename, bIsDirectory);
+			if (!bKeepIterating)
+			{
+				return false;
+			}
 		}
+		// Continue iterating.
+		return true;
 	}
 };
 
@@ -809,11 +818,11 @@ bool FSandboxPlatformFile::IterateDirectory(const TCHAR* Directory, IPlatformFil
 	FString SandboxDir = ConvertToSandboxPath(Directory);
 	bool Result = LowerLevel->IterateDirectory(*SandboxDir, SandboxVisitor);
 	// don't iterate the same directory twice if the Convert didn't change the path, or if we were asked to stop iterating
-	if (Result && (SandboxDir != Directory))
+	if (SandboxVisitor.bKeepIterating && (SandboxDir != Directory))
 	{
-		Result = LowerLevel->IterateDirectory(Directory, SandboxVisitor);
+		Result = LowerLevel->IterateDirectory(Directory, SandboxVisitor) || Result;
 	}
-	return Result;
+	return Result && SandboxVisitor.bKeepIterating;
 }
 
 bool FSandboxPlatformFile::IterateDirectoryRecursively(const TCHAR* Directory, IPlatformFile::FDirectoryVisitor& Visitor)
@@ -822,11 +831,11 @@ bool FSandboxPlatformFile::IterateDirectoryRecursively(const TCHAR* Directory, I
 	FString SandboxDir = ConvertToSandboxPath(Directory);
 	bool Result = LowerLevel->IterateDirectoryRecursively(*SandboxDir, SandboxVisitor);
 	// don't iterate the same directory twice if the Convert didn't change the path, or if we were asked to stop iterating
-	if (Result && (SandboxDir != Directory))
+	if (SandboxVisitor.bKeepIterating && (SandboxDir != Directory))
 	{
-		Result = LowerLevel->IterateDirectoryRecursively(Directory, SandboxVisitor);
+		Result = LowerLevel->IterateDirectoryRecursively(Directory, SandboxVisitor) || Result;
 	}
-	return Result;
+	return Result && SandboxVisitor.bKeepIterating;
 }
 
 class FSandboxStatVisitor : public IPlatformFile::FDirectoryStatVisitor
@@ -835,7 +844,8 @@ public:
 	FDirectoryStatVisitor& Visitor;
 	FSandboxPlatformFile& SandboxFile;
 	TSet<FString> VisitedSandboxFiles;
-	bool bIsRecursive;
+	bool bIsRecursive = false;
+	bool bKeepIterating = true;
 
 	FSandboxStatVisitor(FDirectoryStatVisitor& InVisitor, FSandboxPlatformFile& InSandboxFile, bool bInIsRecursive)
 		: Visitor(InVisitor)
@@ -862,12 +872,20 @@ public:
 			{
 				// "fake" the injected directory so everything falls into place
 				FFileStatData InjectedStat = SandboxFile.LowerLevel->GetStatData(*SandboxFile.InjectedSourceDirectory);
-				Visitor.CallShouldVisitAndVisit(*SandboxFile.InjectedSourceDirectory, InjectedStat);
+				bKeepIterating = bKeepIterating && Visitor.CallShouldVisitAndVisit(*SandboxFile.InjectedSourceDirectory, InjectedStat);
+				if (!bKeepIterating)
+				{
+					return false;
+				}
 
 				// for recursive visiting, we need to recurse into the forced directory
 				if (bIsRecursive)
 				{
 					SandboxFile.LowerLevel->IterateDirectoryStatRecursively(*SandboxFile.InjectedTargetDirectory, *this);
+					if (!bKeepIterating)
+					{
+						return false;
+					}
 				}
 			}
 
@@ -907,14 +925,14 @@ public:
 
 		if (CanVisit)
 		{
-			bool Result = Visitor.CallShouldVisitAndVisit(*LocalFilename, StatData);
-			return Result;
-		}
-		else
-		{
-			// Continue iterating.
-			return true;
+			bKeepIterating = bKeepIterating && Visitor.CallShouldVisitAndVisit(*LocalFilename, StatData);
+			if (!bKeepIterating)
+			{
+				return false;
+			}
 		}
+		// Continue iterating.
+		return true;
 	}
 };
 
@@ -924,11 +942,11 @@ bool		FSandboxPlatformFile::IterateDirectoryStat(const TCHAR* Directory, IPlatfo
 	FString SandboxDir = ConvertToSandboxPath(Directory);
 	bool Result = LowerLevel->IterateDirectoryStat(*SandboxDir, SandboxVisitor);
 	// don't iterate the same directory twice if the Convert didn't change the path, or if we were asked to stop iterating (which can happen with injection type)
-	if (Result && (SandboxDir != Directory))
+	if (SandboxVisitor.bKeepIterating && (SandboxDir != Directory))
 	{
-		Result = LowerLevel->IterateDirectoryStat(Directory, SandboxVisitor);
+		Result = LowerLevel->IterateDirectoryStat(Directory, SandboxVisitor) || Result;
 	}
-	return Result;
+	return Result && SandboxVisitor.bKeepIterating;
 }
 
 bool		FSandboxPlatformFile::IterateDirectoryStatRecursively(const TCHAR* Directory, IPlatformFile::FDirectoryStatVisitor& Visitor)
@@ -937,11 +955,11 @@ bool		FSandboxPlatformFile::IterateDirectoryStatRecursively(const TCHAR* Directo
 	FString SandboxDir = ConvertToSandboxPath(Directory);
 	bool Result = LowerLevel->IterateDirectoryStatRecursively(*ConvertToSandboxPath(Directory), SandboxVisitor);
 	// don't iterate the same directory twice if the Convert didn't change the path, or if we were asked to stop iterating (which can happen with injection type)
-	if (Result && (SandboxDir != Directory))
+	if (SandboxVisitor.bKeepIterating && (SandboxDir != Directory))
 	{
-		Result = LowerLevel->IterateDirectoryStatRecursively(Directory, SandboxVisitor);
+		Result = LowerLevel->IterateDirectoryStatRecursively(Directory, SandboxVisitor) || Result;
 	}
-	return Result;
+	return Result && SandboxVisitor.bKeepIterating;
 }
 
 bool		FSandboxPlatformFile::DeleteDirectoryRecursively(const TCHAR* Directory)
diff --git a/Engine/Source/Runtime/StorageServerClient/Private/StorageServerPlatformFile.cpp b/Engine/Source/Runtime/StorageServerClient/Private/StorageServerPlatformFile.cpp
--- a/Engine/Source/Runtime/StorageServerClient/Private/StorageServerPlatformFile.cpp
+++ b/Engine/Source/Runtime/StorageServerClient/Private/StorageServerPlatformFile.cpp
@@ -1643,6 +1643,7 @@ private:
 	ParentVisitorClass& RealVisitor;
 	IPlatformFile* LowerLevel;
 	FStorageServerPlatformFile& ServerPlatformFile;
+	bool bKeepIterating = true;
 
 public:
 	FUniqueDirectoryStatVisitor(FStorageServerPlatformFile* InPlatformFile, ParentVisitorClass& InVisitor)
@@ -1652,12 +1653,12 @@ public:
 	{
 	}
 
-	virtual bool ShouldVisitLeafPathname(FStringView LeafPathname)
+	virtual bool ShouldVisitLeafPathname(FStringView LeafPathname) override
 	{
 		return RealVisitor.ShouldVisitLeafPathname(LeafPathname);
 	}
 
-	bool Visit(const TCHAR* FilenameOrDirectory, DataType StatData)
+	virtual bool Visit(const TCHAR* FilenameOrDirectory, DataType StatData) override
 	{
 		// for speed reasons, we only do the double location checks when we enable the new
 		// PreferLocal behavior - if we ever have duped results without it, we can remove 
@@ -1672,22 +1673,26 @@ public:
 			}
 			AlreadyVisited.Add(FileStr);
 		}
-		return RealVisitor.Visit(FilenameOrDirectory, StatData);
+		bKeepIterating = bKeepIterating && RealVisitor.Visit(FilenameOrDirectory, StatData);
+		return bKeepIterating;
 	}
 
-
 	bool PerformMergedIteration(const TCHAR* Directory,
 		TFunctionRef<bool(const TCHAR* Directory, ParentVisitorClass& Visitor)> LowLevelOperation,
 		TFunctionRef<bool(const TCHAR* Directory, ParentVisitorClass& Visitor)> RemoteOperation)
 	{
 		// first, if we prefer local assets, iterate on local first (in case the data doesn't match remote for whatever reason,
 		// local data will be used)
+		bool bLocalDirectoryExists = false;
+		bool bRemoteDirectoryExists = false;
 		if (GPreferLocalForNonAssets)
 		{
-			if (LowLevelOperation(Directory, *this) == false)
+			bool bExistsAndKeepIterating = LowLevelOperation(Directory, *this);
+			if (!bKeepIterating)
 			{
 				return false;
 			}
+			bLocalDirectoryExists = bExistsAndKeepIterating;
 		}
 
 		// then look remote
@@ -1696,46 +1701,50 @@ public:
 		bool bIterateOnServer = ServerPlatformFile.MakeStorageServerPath(Directory, StorageServerDirectory) && ServerPlatformFile.ServerToc.DirectoryExists(*StorageServerDirectory);
 		if (bIterateOnServer)
 		{
-			if (RemoteOperation(*StorageServerDirectory, *this) == false)
+			bool bExistsAndKeepIterating = RemoteOperation(*StorageServerDirectory, *this);
+			if (!bKeepIterating)
 			{
 				return false;
 			}
+			bRemoteDirectoryExists = bExistsAndKeepIterating;
 		}
 
 		// finall look locally if we are preferring remote over local
 		if (!GPreferLocalForNonAssets)
 		{
-			if (LowLevelOperation(Directory, *this) == false)
+			bool bExistsAndKeepIterating = LowLevelOperation(Directory, *this);
+			if (!bKeepIterating)
 			{
 				return false;
 			}
+			bLocalDirectoryExists = bExistsAndKeepIterating;
 		}
 
-		return true;
+		return bKeepIterating && (bLocalDirectoryExists || bRemoteDirectoryExists);
 	}
 
 };
 
 
-
 bool FStorageServerPlatformFile::IterateDirectory(const TCHAR* Directory, IPlatformFile::FDirectoryVisitor& Visitor)
 {
 	FUniqueDirectoryStatVisitor<FDirectoryVisitor, bool> MergedVisitor(this, Visitor);
 
 	return MergedVisitor.PerformMergedIteration(Directory,
-		[this](const TCHAR* Directory, FDirectoryVisitor& Visitor)
+		[this](const TCHAR* Directory, FDirectoryVisitor& InMergedVisitor)
 		{
-			return LowerLevel ? LowerLevel->IterateDirectory(Directory, Visitor) : true;
+			return LowerLevel ? LowerLevel->IterateDirectory(Directory, InMergedVisitor) : false;
 		},
-		[this](const TCHAR* Directory, FDirectoryVisitor& Visitor)
+		[this](const TCHAR* Directory, FDirectoryVisitor& InMergedVisitor)
 		{
-			return ServerToc.IterateDirectory(Directory, [this, &Visitor](const FIoChunkId& FileChunkId, const TCHAR* FilenameOrDirectory, int64 RawSize)
+			return ServerToc.IterateDirectory(Directory, [this, &InMergedVisitor]
+				(const FIoChunkId& FileChunkId, const TCHAR* FilenameOrDirectory, int64 RawSize)
 				{
 					TStringBuilder<1024> LocalPath;
 					bool bConverted = MakeLocalPath(FilenameOrDirectory, LocalPath);
 					check(bConverted);
 					const bool bDirectory = !FileChunkId.IsValid();
-					return Visitor.CallShouldVisitAndVisit(*LocalPath, bDirectory);
+					return InMergedVisitor.CallShouldVisitAndVisit(*LocalPath, bDirectory);
 				});
 		}
 	);
@@ -1746,37 +1755,38 @@ bool FStorageServerPlatformFile::IterateDirectoryRecursively(const TCHAR* Direct
 	FUniqueDirectoryStatVisitor<FDirectoryVisitor, bool> MergedVisitor(this, Visitor);
 
 	return MergedVisitor.PerformMergedIteration(Directory,
-		[this](const TCHAR* Directory, FDirectoryVisitor& Visitor)
+		[this](const TCHAR* Directory, FDirectoryVisitor& InMergedVisitor)
 		{
-			return LowerLevel ? LowerLevel->IterateDirectoryRecursively(Directory, Visitor) : true;
+			return LowerLevel ? LowerLevel->IterateDirectoryRecursively(Directory, InMergedVisitor) : false;
 		},
-		[this](const TCHAR* Directory, FDirectoryVisitor& Visitor)
+		[this](const TCHAR* Directory, FDirectoryVisitor& InMergedVisitor)
 		{
-			return ServerToc.IterateDirectoryRecursively(Directory, [this, &Visitor](const FIoChunkId& FileChunkId, const TCHAR* FilenameOrDirectory, int64 RawSize)
+			return ServerToc.IterateDirectoryRecursively(Directory, [this, &InMergedVisitor]
+				(const FIoChunkId& FileChunkId, const TCHAR* FilenameOrDirectory, int64 RawSize)
 				{
 					TStringBuilder<1024> LocalPath;
 					bool bConverted = MakeLocalPath(FilenameOrDirectory, LocalPath);
 					check(bConverted);
 					const bool bDirectory = !FileChunkId.IsValid();
-					return Visitor.CallShouldVisitAndVisit(*LocalPath, bDirectory);
+					return InMergedVisitor.CallShouldVisitAndVisit(*LocalPath, bDirectory);
 				});
 		}
 	);
 }
 
-
 bool FStorageServerPlatformFile::IterateDirectoryStat(const TCHAR* Directory, FDirectoryStatVisitor& Visitor)
 {
 	FUniqueDirectoryStatVisitor<FDirectoryStatVisitor, const FFileStatData&> MergedVisitor(this, Visitor);
 
 	return MergedVisitor.PerformMergedIteration(Directory,
-		[this](const TCHAR* Directory, FDirectoryStatVisitor& Visitor)
+		[this](const TCHAR* Directory, FDirectoryStatVisitor& InMergedVisitor)
 		{
-			return LowerLevel ? LowerLevel->IterateDirectoryStat(Directory, Visitor) : true;
+			return LowerLevel ? LowerLevel->IterateDirectoryStat(Directory, InMergedVisitor) : false;
 		},
-		[this](const TCHAR* Directory, FDirectoryStatVisitor& Visitor)
+		[this](const TCHAR* Directory, FDirectoryStatVisitor& InMergedVisitor)
 		{
-			return ServerToc.IterateDirectory(Directory, [this, &Visitor](const FIoChunkId& FileChunkId, const TCHAR* ServerFilenameOrDirectory, int64 RawSize)
+			return ServerToc.IterateDirectory(Directory, [this, &InMergedVisitor]
+				(const FIoChunkId& FileChunkId, const TCHAR* ServerFilenameOrDirectory, int64 RawSize)
 				{
 					TStringBuilder<1024> LocalPath;
 					bool bConverted = MakeLocalPath(ServerFilenameOrDirectory, LocalPath);
@@ -1803,7 +1813,7 @@ bool FStorageServerPlatformFile::IterateDirectoryStat(const TCHAR* Directory, FD
 							true,
 							true);
 					}
-					return Visitor.CallShouldVisitAndVisit(*LocalPath, FileStatData); 
+					return InMergedVisitor.CallShouldVisitAndVisit(*LocalPath, FileStatData);
 				});
 		}
 	);
